Copying //UE4/Dev-Sequencer to //UE4/Main (Source: //UE4/Dev-Sequencer @ 3617642)

diff --git a/Engine/Plugins/Media/MfMedia/Source/MfMedia/Private/Player/MfMediaPlayer.cpp b/Engine/Plugins/Media/MfMedia/Source/MfMedia/Private/Player/MfMediaPlayer.cpp
--- a/Engine/Plugins/Media/MfMedia/Source/MfMedia/Private/Player/MfMediaPlayer.cpp
+++ b/Engine/Plugins/Media/MfMedia/Source/MfMedia/Private/Player/MfMediaPlayer.cpp
@@ -1,12 +1,15 @@
 // Copyright 1998-2017 Epic Games, Inc. All Rights Reserved.
 
 #include "MfMediaPlayer.h"
-#include "Serialization/ArrayReader.h"
-#include "Misc/FileHelper.h"
 
 #if MFMEDIA_SUPPORTED_PLATFORM
 
-#include "MfMediaResolver.h"
+#include "Async/Async.h"
+#include "IMediaEventSink.h"
+#include "IMediaOptions.h"
+#include "MediaSamples.h"
+
+#include "MfMediaSourceReaderCallback.h"
 #include "MfMediaTracks.h"
 #include "MfMediaUtils.h"
 
@@ -17,21 +20,28 @@
 	#include "XboxOneAllowPlatformTypes.h"
 #endif
 
+#define MFMEDIAPLAYER_USE_SEEKANDREVERSE 0 // not fully working yet
+
+#define LOCTEXT_NAMESPACE "FMfMediaPlayer"
+
 
 /* FMfVideoPlayer structors
  *****************************************************************************/
 
-FMfMediaPlayer::FMfMediaPlayer()
-	: CanScrub(false)
-	, Characteristics(0)
+FMfMediaPlayer::FMfMediaPlayer(IMediaEventSink& InEventSink)
+	: Characteristics(0)
+	, CurrentDuration(FTimespan::Zero())
 	, CurrentRate(0.0f)
+	, CurrentState(EMediaState::Closed)
+	, CurrentStatus(EMediaStatus::None)
 	, CurrentTime(FTimespan::Zero())
-	, Duration(0)
-	, Looping(false)
-{
-	Resolver = new FMfMediaResolver;
-	Tracks = new FMfMediaTracks;
-}
+	, EventSink(InEventSink)
+	, PlaybackRestarted(true)
+	, Samples(MakeShared<FMediaSamples, ESPMode::ThreadSafe>())
+	, ShouldLoop(false)
+	, SourceReaderError(false)
+	, Tracks(MakeShared<FMfMediaTracks, ESPMode::ThreadSafe>())
+{ }
 
 
 FMfMediaPlayer::~FMfMediaPlayer()
@@ -40,526 +50,718 @@ FMfMediaPlayer::~FMfMediaPlayer()
 }
 
 
-/* IMediaControls interface
+/* IMediaPlayer interface
  *****************************************************************************/
 
-FTimespan FMfMediaPlayer::GetDuration() const
+void FMfMediaPlayer::Close()
 {
-	return Duration;
+	if ((CurrentState == EMediaState::Closed))
+	{
+		return;
+	}
+
+	// reset player
+	if (MediaSource != NULL)
+	{
+		MediaSource->Shutdown();
+		MediaSource = NULL;
+	}
+
+	Tracks->Shutdown();
+
+	PresentationDescriptor.Reset();
+	RateControl.Reset();
+	RateSupport.Reset();
+	SourceReader.Reset();
+	SourceReaderCallback.Reset();
+
+	Characteristics = 0;
+	CurrentDuration = FTimespan::Zero();
+	CurrentRate = 0.0f;
+	CurrentState = EMediaState::Closed;
+	CurrentStatus = EMediaStatus::None;
+	CurrentTime = FTimespan::Zero();
+	MediaUrl = FString();
+//	PendingSeekTime.Reset();
+	PlaybackRestarted = true;
+	SourceReaderError = false;
+	ThinnedRates.Empty();
+	UnthinnedRates.Empty();
+
+	// notify listeners
+	EventSink.ReceiveMediaEvent(EMediaEvent::TracksChanged);
+	EventSink.ReceiveMediaEvent(EMediaEvent::MediaClosed);
 }
 
 
-float FMfMediaPlayer::GetRate() const
+IMediaCache& FMfMediaPlayer::GetCache()
 {
-	return CurrentRate;
+	return *this;
 }
 
 
-EMediaState FMfMediaPlayer::GetState() const
+IMediaControls& FMfMediaPlayer::GetControls()
 {
-	if (Resolver->IsResolving())
-	{
-		return EMediaState::Preparing;
-	}
+	return *this;
+}
 
-	if (SourceReader == NULL)
-	{
-		return EMediaState::Closed;
-	}
 
-	if (CurrentRate == 0.0f)
-	{
-		if (CurrentTime == FTimespan::Zero())
-		{
-			return EMediaState::Stopped;
-		}
+FString FMfMediaPlayer::GetInfo() const
+{
+	return Tracks->GetInfo();
+}
 
-		return EMediaState::Paused;
-	}
 
-	return EMediaState::Playing;
+FName FMfMediaPlayer::GetName() const
+{
+	static FName PlayerName(TEXT("MfMedia"));
+	return PlayerName;
 }
 
 
-TRange<float> FMfMediaPlayer::GetSupportedRates(EMediaPlaybackDirections Direction, bool Unthinned) const
+IMediaSamples& FMfMediaPlayer::GetSamples()
 {
-	if (RateSupport == NULL)
-	{
-		return TRange<float>(0.0f);
-	}
+	check(Samples.IsValid());
+	return*Samples.Get();;
+}
 
-	float MaxRate = 0.0f;
-	float MinRate = 0.0f;
 
-	switch (Direction)
-	{
-	case EMediaPlaybackDirections::Forward:
-		RateSupport->GetSlowestRate(MFRATE_FORWARD, !Unthinned, &MinRate);
-		RateSupport->GetFastestRate(MFRATE_FORWARD, !Unthinned, &MaxRate);
-		break;
+FString FMfMediaPlayer::GetStats() const
+{
+	FString Result;
+	Tracks->AppendStats(Result);
+
+	return Result;
+}
 
-	case EMediaPlaybackDirections::Reverse:
-		RateSupport->GetSlowestRate(MFRATE_REVERSE, !Unthinned, &MaxRate);
-		RateSupport->GetFastestRate(MFRATE_REVERSE, !Unthinned, &MinRate);
-		break;
-	}
 
-	return TRange<float>(MinRate, MaxRate);
+IMediaTracks& FMfMediaPlayer::GetTracks()
+{
+	return *Tracks;
 }
 
 
-FTimespan FMfMediaPlayer::GetTime() const
+FString FMfMediaPlayer::GetUrl() const
 {
-	return CurrentTime;
+	return MediaUrl;
 }
 
 
-bool FMfMediaPlayer::IsLooping() const
+IMediaView& FMfMediaPlayer::GetView()
 {
-	return Looping;
+	return *this;
 }
 
 
-bool FMfMediaPlayer::Seek(const FTimespan& Time)
+bool FMfMediaPlayer::Open(const FString& Url, const IMediaOptions* Options)
 {
-	if (SourceReader == NULL)
+	Close();
+
+	if (Url.IsEmpty())
 	{
 		return false;
 	}
 
-	PROPVARIANT Position;
+	const bool Precache = (Options != nullptr) ? Options->GetMediaOption("PrecacheFile", false) : false;
 
-	if (Time >= FTimespan::Zero())
-	{
-		Position.vt = VT_I8;
-		Position.hVal.QuadPart = Time.GetTicks();
-	}
-	else
+	return InitializePlayer(nullptr, Url, Precache);
+}
+
+
+bool FMfMediaPlayer::Open(const TSharedRef<FArchive, ESPMode::ThreadSafe>& Archive, const FString& OriginalUrl, const IMediaOptions* /*Options*/)
+{
+	Close();
+
+	if (Archive->TotalSize() == 0)
 	{
-		PropVariantInit(&Position);
+		UE_LOG(LogMfMedia, Error, TEXT("Cannot open media from archive (archive is empty)."));
+		return false;
 	}
 
-	if (FAILED(SourceReader->SetCurrentPosition(GUID_NULL, Position)))
+	if (OriginalUrl.IsEmpty())
 	{
+		UE_LOG(LogMfMedia, Error, TEXT("Cannot open media from archive (no original URL provided)."));
 		return false;
 	}
 
-	CurrentTime = Time;
+	return InitializePlayer(Archive, OriginalUrl, false);
+}
 
-	return true;
+
+void FMfMediaPlayer::TickAudio()
+{
+	if (CurrentState == EMediaState::Playing)
+	{
+		Tracks->TickAudio(CurrentRate, CurrentTime);
+	}
 }
 
 
-bool FMfMediaPlayer::SetLooping(bool InLooping)
+void FMfMediaPlayer::TickFetch(FTimespan /*DeltaTime*/)
 {
-	Looping = InLooping;
+	bool MediaSourceChanged = false;
+	bool TrackSelectionChanged = false;
 
-	return true;
+	Tracks->GetFlags(MediaSourceChanged, TrackSelectionChanged);
+
+	if (MediaSourceChanged)
+	{
+		if (Tracks->IsInitialized())
+		{
+			if (CurrentState == EMediaState::Preparing)
+			{
+				CurrentDuration = Tracks->GetDuration();
+				MediaSource = Tracks->GetMediaSource();
+				SourceReader = Tracks->GetSourceReader();
+
+				UpdateCharacteristics();
+
+				if (MediaSource.IsValid())
+				{
+					CurrentState = EMediaState::Stopped;
+					EventSink.ReceiveMediaEvent(EMediaEvent::MediaOpened);
+				}
+				else
+				{
+					CurrentState = EMediaState::Error;
+					EventSink.ReceiveMediaEvent(EMediaEvent::MediaOpenFailed);
+				}
+			}
+
+			EventSink.ReceiveMediaEvent(EMediaEvent::TracksChanged);
+		}
+		else
+		{
+			CurrentState = EMediaState::Error;
+			EventSink.ReceiveMediaEvent(EMediaEvent::MediaOpenFailed);
+		}
+	}
+
+	if (TrackSelectionChanged)
+	{
+		if (CommitTime(CurrentTime))
+		{
+			Tracks->Restart();
+		}
+	}
+
+	if (MediaSourceChanged || TrackSelectionChanged)
+	{
+		Tracks->ClearFlags();
+	}
 }
 
 
-bool FMfMediaPlayer::SetRate(float Rate)
+void FMfMediaPlayer::TickInput(FTimespan DeltaTime)
 {
-	if ((MediaSource == NULL) || !SupportsRate(Rate, false))
+	if ((CurrentState != EMediaState::Playing) || (CurrentDuration == FTimespan::Zero()))
 	{
-		return false;
+		return;
 	}
 
-	if (Rate == CurrentRate)
+	if (SourceReaderError)
 	{
-		return true;
+		CurrentState = EMediaState::Error;
+
+		return;
 	}
 
-	if (FMath::IsNearlyZero(Rate))
+	// update clock
+	if (PlaybackRestarted)
 	{
-		// Make sure we set enabled to false before pausing or the render thread might get stuck waiting for another sample to read.
-		Tracks->SetEnabled(false);
-
-		if (FAILED(MediaSource->Pause()))
-		{
-			return false;
-		}
+		PlaybackRestarted = false;
 	}
 	else
 	{
-		if (Rate != 1.0f)
-		{
-			if ((RateControl == NULL) || FAILED(RateControl->SetRate(SupportsRate(Rate, true) ? FALSE : TRUE, Rate)))
-			{
-				return false;
-			}
-		}
+		CurrentTime += DeltaTime * CurrentRate;
+	}
 
-		if (FMath::IsNearlyZero(CurrentRate))
-		{
-			PROPVARIANT StartPosition;
+	// handle looping
+	if ((CurrentTime >= CurrentDuration) || (CurrentTime < FTimespan::Zero()))
+	{
+		EventSink.ReceiveMediaEvent(EMediaEvent::PlaybackEndReached);
 
-			if (CurrentTime >= FTimespan::Zero())
+		if (ShouldLoop)
+		{
+			if (CurrentRate > 0.0f)
 			{
-				StartPosition.vt = VT_I8;
-				StartPosition.hVal.QuadPart = CurrentTime.GetTicks();
+				CurrentTime = FTimespan::Zero();
 			}
 			else
 			{
-				PropVariantInit(&StartPosition);
+				CurrentTime = CurrentDuration - FTimespan(1);
 			}
 
-			if (FAILED(MediaSource->Start(NULL, &GUID_NULL, &StartPosition)))
-			{
-				return false;
-			}
+			PlaybackRestarted = true;
+		}
+		else
+		{
+			CurrentRate = 0.0f;
+			CurrentState = EMediaState::Stopped;
+			CurrentTime = FTimespan::Zero();
+
+			EventSink.ReceiveMediaEvent(EMediaEvent::PlaybackSuspended);
+		}
 
-			Tracks->SetEnabled(true);
+		if (CommitTime(CurrentTime))
+		{
+			Tracks->Restart();
 		}
 	}
+	else if (CurrentRate < 0.0f)
+	{
+		// IMFSourceReader does not support reverse playback, even if the media
+		// source does. we emulate the WMF behavior by seeking to key frames.
 
-	CurrentRate = Rate;
+		if (!CommitTime(CurrentTime))
+		{
+			CurrentState = EMediaState::Error;
+		}
+	}
 
-	return true;
+	// update tracks
+	if (CurrentState == EMediaState::Playing)
+	{
+		Tracks->TickInput(CurrentRate, CurrentTime, DeltaTime);
+	}
 }
 
 
-bool FMfMediaPlayer::SupportsRate(float Rate, bool Unthinned) const
+/* FMfMediaPlayer implementation
+ *****************************************************************************/
+
+bool FMfMediaPlayer::CommitTime(FTimespan Time)
 {
-	// @todo gmp: reverse playback doesn't seem to work correctly and breaks the source reader
-	if (Rate < 0.0f)
+	// perform seek
+	PROPVARIANT Position;
+
+	if (Time >= FTimespan::Zero())
 	{
+		Position.vt = VT_I8;
+		Position.hVal.QuadPart = Time.GetTicks();
+	}
+	else
+	{
+		::PropVariantInit(&Position);
+	}
+
+	const HRESULT Result = SourceReader->SetCurrentPosition(GUID_NULL, Position);
+
+	if (FAILED(Result))
+	{
+		UE_LOG(LogMfMedia, Warning, TEXT("Failed to commit time %s: %s"), *Time.ToString(), *MfMedia::ResultToString(Result));
 		return false;
 	}
 
-	if (Rate == 1.0f)
+	CurrentTime = Time;
+
+	return true;
+}
+
+
+bool FMfMediaPlayer::InitializePlayer(const TSharedPtr<FArchive, ESPMode::ThreadSafe>& Archive, const FString& Url, bool Precache)
+{
+	UE_LOG(LogMfMedia, VeryVerbose, TEXT("Player %llx: Initializing %s (archive = %s, precache = %s)"), this, *Url, Archive.IsValid() ? TEXT("yes") : TEXT("no"), Precache ? TEXT("yes") : TEXT("no"));
+
+	CurrentState = EMediaState::Preparing;
+
+	MediaUrl = Url;
+	SourceReaderCallback = new FMfMediaSourceReaderCallback(*this);
+
+	// initialize presentation on a separate thread
+	const EAsyncExecution Execution = Precache ? EAsyncExecution::Thread : EAsyncExecution::ThreadPool;
+
+	Async<void>(Execution, [
+		Archive, Url, Precache,
+		Callback = TComPtr<FMfMediaSourceReaderCallback>(SourceReaderCallback),
+		SamplesPtr = TWeakPtr<FMediaSamples, ESPMode::ThreadSafe>(Samples),
+		TracksPtr = TWeakPtr<FMfMediaTracks, ESPMode::ThreadSafe>(Tracks)]()
 	{
-		return true;
+		TSharedPtr<FMediaSamples, ESPMode::ThreadSafe> PinnedSamples = SamplesPtr.Pin();
+		TSharedPtr<FMfMediaTracks, ESPMode::ThreadSafe> PinnedTracks= TracksPtr.Pin();
+
+		if (PinnedSamples.IsValid() && PinnedTracks.IsValid())
+		{
+			TComPtr<IMFMediaSource> MediaSource = MfMedia::ResolveMediaSource(Archive, Url, Precache);
+			PinnedTracks->Initialize(MediaSource, Callback, PinnedSamples.ToSharedRef());
+		}
+	});
+
+	return true;
+}
+
+
+void FMfMediaPlayer::UpdateCharacteristics()
+{
+	if (!MediaSource.IsValid())
+	{
+		return;
 	}
 
-	if (Rate == 0.0f)
+	// get characteristics
+	HRESULT Result = MediaSource->GetCharacteristics(&Characteristics);
+
+	if (FAILED(Result))
 	{
-		return ((Characteristics & MFMEDIASOURCE_CAN_SEEK) != 0);
+		UE_LOG(LogMfMedia, Warning, TEXT("Failed to get media source characteristics: %s"), *MfMedia::ResultToString(Result));
 	}
 
-	if ((RateControl == NULL) || (RateSupport == NULL))
+	// get service interface
+	TComPtr<IMFGetService> GetService;
+
+	Result = MediaSource->QueryInterface(IID_PPV_ARGS(&GetService));
+
+	if (FAILED(Result))
 	{
-		return false;
+		UE_LOG(LogMfMedia, Warning, TEXT("Failed to query service interface: %s"), *MfMedia::ResultToString(Result));
+		return;
 	}
 
-	if (Unthinned)
+	// get rate control & rate support, if available
+	Result = GetService->GetService(MF_RATE_CONTROL_SERVICE, IID_PPV_ARGS(&RateControl));
+
+	if (FAILED(Result))
 	{
-		return SUCCEEDED(RateSupport->IsRateSupported(FALSE, Rate, NULL));
+		UE_LOG(LogMfMedia, Log, TEXT("Rate control service unavailable (%s)"), *MfMedia::ResultToString(Result));
+	}
+	else
+	{
+		UE_LOG(LogMfMedia, Verbose, TEXT("Player %llx: Rate control ready"), this);
 	}
 
-	return SUCCEEDED(RateSupport->IsRateSupported(TRUE, Rate, NULL));
-}
+	Result = GetService->GetService(MF_RATE_CONTROL_SERVICE, IID_PPV_ARGS(&RateSupport));
 
+	if (FAILED(Result))
+	{
+		UE_LOG(LogMfMedia, Log, TEXT("Rate support service unavailable (%s)"), *MfMedia::ResultToString(Result));
+	}
+	else
+	{
+		UE_LOG(LogMfMedia, Verbose, TEXT("Player %llx: Rate support ready"), this);
+	}
 
-bool FMfMediaPlayer::SupportsScrubbing() const
-{
-	return CanScrub && ((Characteristics & MFMEDIASOURCE_HAS_SLOW_SEEK) == 0);
-}
+	// cache rate control properties
+	if (RateSupport.IsValid())
+	{
+		float MaxRate = 0.0f;
+		float MinRate = 0.0f;
 
+		if (SUCCEEDED(RateSupport->GetSlowestRate(MFRATE_FORWARD, TRUE, &MinRate)) &&
+			SUCCEEDED(RateSupport->GetFastestRate(MFRATE_FORWARD, TRUE, &MaxRate)))
+		{
+			ThinnedRates.Add(TRange<float>::Inclusive(MinRate, MaxRate));
+		}
 
-bool FMfMediaPlayer::SupportsSeeking() const
-{
-	return (((Characteristics & MFMEDIASOURCE_CAN_SEEK) != 0) && (Duration > FTimespan::Zero()));
+#if MFMEDIAPLAYER_USE_SEEKANDREVERSE
+		if (SUCCEEDED(RateSupport->GetSlowestRate(MFRATE_REVERSE, TRUE, &MinRate)) &&
+			SUCCEEDED(RateSupport->GetFastestRate(MFRATE_REVERSE, TRUE, &MaxRate)))
+		{
+			ThinnedRates.Add(TRange<float>::Inclusive(MaxRate, MinRate));
+		}
+#endif
+
+		if (SUCCEEDED(RateSupport->GetSlowestRate(MFRATE_FORWARD, FALSE, &MinRate)) &&
+			SUCCEEDED(RateSupport->GetFastestRate(MFRATE_FORWARD, FALSE, &MaxRate)))
+		{
+			UnthinnedRates.Add(TRange<float>::Inclusive(MinRate, MaxRate));
+		}
+
+#if MFMEDIAPLAYER_USE_SEEKANDREVERSE
+		if (SUCCEEDED(RateSupport->GetSlowestRate(MFRATE_REVERSE, FALSE, &MinRate)) &&
+			SUCCEEDED(RateSupport->GetFastestRate(MFRATE_REVERSE, FALSE, &MaxRate)))
+		{
+			UnthinnedRates.Add(TRange<float>::Inclusive(MaxRate, MinRate));
+		}
+#endif
+	}
+	else
+	{
+		ThinnedRates.Empty();
+		UnthinnedRates.Empty();
+	}
 }
 
 
-/* IMediaPlayer interface
+/* IMediaControls interface
  *****************************************************************************/
 
-void FMfMediaPlayer::Close()
+bool FMfMediaPlayer::CanControl(EMediaControl Control) const
 {
-	Resolver->Cancel();
-
 	if (SourceReader == NULL)
 	{
-		return;
+		return false;
 	}
 
-	// release media source
-	if (MediaSource != NULL)
+	if (Control == EMediaControl::Pause)
 	{
-		MediaSource->Shutdown();
-		MediaSource = NULL;
+		return ((CurrentState == EMediaState::Playing) && ((Characteristics & MFMEDIASOURCE_CAN_PAUSE) != 0));
 	}
 
-	// reset player
-	CanScrub = false;
-	Characteristics = 0;
-	CurrentRate = 0.0f;
-	CurrentTime = FTimespan::Zero();
-	Duration = 0;
-	Info.Empty();
-	MediaUrl = FString();
-	RateControl = NULL;
-	RateSupport = NULL;
-	SourceReader = NULL;
-	Tracks->Reset();
+	if (Control == EMediaControl::Resume)
+	{
+		return ((CurrentState != EMediaState::Playing) && ThinnedRates.Contains(1.0f));
+	}
 
-	// notify listeners
-	MediaEvent.Broadcast(EMediaEvent::TracksChanged);
-	MediaEvent.Broadcast(EMediaEvent::MediaClosed);
+#if MFMEDIAPLAYER_USE_SEEKANDREVERSE
+	if (Control == EMediaControl::Scrub)
+	{
+		return (((Characteristics & MFMEDIASOURCE_HAS_SLOW_SEEK) == 0) && (ThinnedRates.Contains(0.0f)));
+	}
+
+	if (Control == EMediaControl::Seek)
+	{
+		return (((Characteristics & MFMEDIASOURCE_CAN_SEEK) != 0) && (CurrentDuration > FTimespan::Zero()));
+	}
+#endif
+
+	return false;
 }
 
 
-IMediaControls& FMfMediaPlayer::GetControls()
+FTimespan FMfMediaPlayer::GetDuration() const
 {
-	return *this;
+	return CurrentDuration;
 }
 
 
-FString FMfMediaPlayer::GetInfo() const
+float FMfMediaPlayer::GetRate() const
 {
-	return Info;
+	return CurrentRate;
 }
 
 
-FName FMfMediaPlayer::GetName() const
+EMediaState FMfMediaPlayer::GetState() const
 {
-	static FName PlayerName(TEXT("MfMedia"));
-	return PlayerName;
+	return CurrentState;
 }
 
 
-IMediaOutput& FMfMediaPlayer::GetOutput()
+EMediaStatus FMfMediaPlayer::GetStatus() const
 {
-	return *((FMfMediaTracks*)Tracks);
+	return CurrentStatus;
 }
 
 
-FString FMfMediaPlayer::GetStats() const
+TRangeSet<float> FMfMediaPlayer::GetSupportedRates(EMediaRateThinning Thinning) const
 {
-	return TEXT("MfMedia stats information not implemented yet");
+	return (Thinning == EMediaRateThinning::Thinned) ? ThinnedRates : UnthinnedRates;
 }
 
 
-IMediaTracks& FMfMediaPlayer::GetTracks()
+FTimespan FMfMediaPlayer::GetTime() const
 {
-	return *((FMfMediaTracks*)Tracks);
+	return CurrentTime;
 }
 
 
-FString FMfMediaPlayer::GetUrl() const
+bool FMfMediaPlayer::IsLooping() const
 {
-	return MediaUrl;
+	return ShouldLoop;
 }
 
 
-bool FMfMediaPlayer::Open(const FString& Url, const IMediaOptions& Options)
+bool FMfMediaPlayer::Seek(const FTimespan& Time)
 {
-	Close();
+	if (SourceReader == NULL)
+	{
+		return false;
+	}
 
-	if (Url.IsEmpty())
+	// validate seek
+	if ((CurrentState == EMediaState::Closed) || (CurrentState == EMediaState::Error))
 	{
+		UE_LOG(LogMfMedia, Warning, TEXT("Cannot seek while closed or in error state"));
 		return false;
 	}
 
-	// open local files via platform file system
-	if (Url.StartsWith(TEXT("file://")))
+	if ((Time < FTimespan::Zero()) || (Time > CurrentDuration))
 	{
-		TSharedPtr<FArchive, ESPMode::ThreadSafe> Archive;
-		const TCHAR* FilePath = &Url[7];
+		UE_LOG(LogMfMedia, Warning, TEXT("Invalid seek time %s (media duration is %s)"), *Time.ToString(), *CurrentDuration.ToString());
+		return false;
+	}
+	
+	UE_LOG(LogMfMedia, Verbose, TEXT("Player %llx: Seeking to %s"), this, *Time.ToString());
 
-		if (Options.GetMediaOption("PrecacheFile", false))
-		{
-			FArrayReader* Reader = new FArrayReader;
+	if (!CommitTime(Time))
+	{
+		return false;
+	}
 
-			if (FFileHelper::LoadFileToArray(*Reader, FilePath))
-			{
-				Archive = MakeShareable(Reader);
-			}
-			else
-			{
-				delete Reader;
-			}
-		}
-		else
-		{
-			Archive = MakeShareable(IFileManager::Get().CreateFileReader(FilePath));
-		}
+	EventSink.ReceiveMediaEvent(EMediaEvent::SeekCompleted);
 
-		if (!Archive.IsValid())
-		{
-			UE_LOG(LogMfMedia, Warning, TEXT("Failed to open media file: %s"), FilePath);
+	return true;
+}
 
-			return false;
-		}
 
-		return Resolver->ResolveByteStream(Archive.ToSharedRef(), Url, *this);
-	}
+bool FMfMediaPlayer::SetLooping(bool Looping)
+{
+	ShouldLoop = Looping;
 
-	return Resolver->ResolveUrl(Url, *this);
+	return true;
 }
 
 
-bool FMfMediaPlayer::Open(const TSharedRef<FArchive, ESPMode::ThreadSafe>& Archive, const FString& OriginalUrl, const IMediaOptions& Options)
+bool FMfMediaPlayer::SetRate(float Rate)
 {
-	Close();
-
-	if ((Archive->TotalSize() == 0) || OriginalUrl.IsEmpty())
+	if (SourceReader == NULL)
 	{
 		return false;
 	}
 
-	return Resolver->ResolveByteStream(Archive, OriginalUrl, *this);
-}
-
-
-void FMfMediaPlayer::TickPlayer(float DeltaTime)
-{
-	// process deferred tasks
-	TFunction<void()> Task;
-
-	while (GameThreadTasks.Dequeue(Task))
+	if (Rate == CurrentRate)
 	{
-		Task();
+		return true; // rate already set
 	}
 
-	if (GetState() != EMediaState::Playing)
+	if (CurrentDuration == FTimespan::Zero())
 	{
-		return;
+		return false; // nothing to play
 	}
 
-	// extract frame data
-	Tracks->Tick(DeltaTime);
+	BOOL Thin;
 
-	// update playback time
-	CurrentTime += FTimespan::FromSeconds(CurrentRate * DeltaTime);
+	// check whether rate is supported
+	if (UnthinnedRates.Contains(Rate))
+	{
+		UE_LOG(LogMfMedia, Verbose, TEXT("Player %llx: Setting rate from to %f to %f (unthinned)"), this, CurrentRate, Rate);
+		Thin = FALSE;
+	}
+	else if (ThinnedRates.Contains(Rate))
+	{
+		UE_LOG(LogMfMedia, Verbose, TEXT("Player %llx: Setting rate from to %f to %f (thinned)"), this, CurrentRate, Rate);
+		Thin = TRUE;
+	}
+	else
+	{
+		UE_LOG(LogMfMedia, Warning, TEXT("The rate %f is not supported"), Rate);
+		return false;
+	}
 
-	if ((CurrentTime < FTimespan::Zero()) || (CurrentTime > Duration))
+	// set play rate
+	if (RateControl != NULL)
 	{
-		if (Looping)
+		if (Rate <= 0.0f)
 		{
-			// loop to beginning/end
-			if (CurrentRate > 0.0f)
+			// Note: IMFSourceReader does not support reverse playback, even if the
+			// media source does. We switch to scrubbing mode and seek instead
+
+			const HRESULT Result = RateControl->SetRate(FALSE, 0.0f);
+
+			if (FAILED(Result))
 			{
-				Seek(FTimespan::Zero());
+				UE_LOG(LogMfMedia, Warning, TEXT("Failed to commit rate change from %f to zero: %s"), CurrentRate, *MfMedia::ResultToString(Result));
+				return false;
 			}
-			else
+		}
+		else if (!ThinnedRates.IsEmpty())
+		{
+			// Note: for those media sources that support thinning in forward play,
+			// we update the rate controller
+
+			const TCHAR* ThinnedString = Thin ? TEXT("thinned") : TEXT("unthinned");
+			const HRESULT Result = RateControl->SetRate(Thin, Rate);
+
+			if (FAILED(Result))
 			{
-				Seek(Duration - FTimespan(1));
+				UE_LOG(LogMfMedia, Warning, TEXT("Failed to commit rate change from %f to %f (%s): %s"), CurrentRate, Rate, ThinnedString, *MfMedia::ResultToString(Result));
+				return false;
 			}
-
-			Tracks->Reinitialize();
 		}
-		else
+	}
+
+	// handle restarting
+	if ((CurrentRate == 0.0f) && (Rate != 0.0f))
+	{
+		if (CurrentState == EMediaState::Stopped)
 		{
-			// stop playback
-			Seek(FTimespan::Zero());
-			SetRate(0.0f);
-			MediaEvent.Broadcast(EMediaEvent::PlaybackEndReached);
+			if (Rate < 0.0f)
+			{
+				CommitTime(CurrentDuration - FTimespan(1));
+			}
+
+			PlaybackRestarted = true;
+			Tracks->Restart();
 		}
-	}
-}
 
+		CurrentRate = Rate;
+		CurrentState = EMediaState::Playing;
 
-void FMfMediaPlayer::TickVideo(float DeltaTime)
-{
-	// do nothing
-}
+		EventSink.ReceiveMediaEvent(EMediaEvent::PlaybackResumed);
 
+		return true;
+	}
 
-/* IMfMediaResolverCallbacks interface
- *****************************************************************************/
+	// handle pausing
+	if ((CurrentRate != 0.0f) && (Rate == 0.0f))
+	{
+		CurrentRate = Rate;
+		CurrentState = EMediaState::Paused;
 
-void FMfMediaPlayer::ProcessResolveComplete(TComPtr<IUnknown> SourceObject, FString ResolvedUrl)
-{
-	GameThreadTasks.Enqueue([=]() {
-		MediaEvent.Broadcast(
-			FinishOpen(SourceObject, ResolvedUrl)
-				? EMediaEvent::MediaOpened
-				: EMediaEvent::MediaOpenFailed
-		);
-	});
-}
+		EventSink.ReceiveMediaEvent(EMediaEvent::PlaybackSuspended);
 
+		return true;
+	}
 
-void FMfMediaPlayer::ProcessResolveFailed(FString FailedUrl)
-{
-	GameThreadTasks.Enqueue([=]() {
-		MediaEvent.Broadcast(EMediaEvent::MediaOpenFailed);
-	});
+	CurrentRate = Rate;
+
+	return true;
 }
 
 
-/* FMfMediaPlayer implementation
+/* IMfMediaSourceReaderSink interface
  *****************************************************************************/
 
-bool FMfMediaPlayer::FinishOpen(IUnknown* SourceObject, const FString& SourceUrl)
+void FMfMediaPlayer::ReceiveSourceReaderEvent(MediaEventType Event)
 {
-	if (SourceObject == nullptr)
-	{
-		return false;
-	}
-
-	UE_LOG(LogMfMedia, Verbose, TEXT("Initializing media session for %s"), *SourceUrl);
-
-	// create presentation descriptor
-	TComPtr<IMFMediaSource> MediaSourceObject;
+	switch (Event)
 	{
-		HRESULT Result = SourceObject->QueryInterface(IID_PPV_ARGS(&MediaSourceObject));
+	case MEBufferingStarted:
+		CurrentStatus |= EMediaStatus::Buffering;
+		EventSink.ReceiveMediaEvent(EMediaEvent::MediaBuffering);
+		break;
 
-		if (FAILED(Result))
-		{
-			UE_LOG(LogMfMedia, Error, TEXT("Failed to query media source: %s"), *MfMedia::ResultToString(Result));
-			return false;
-		}
-	}
+	case MEBufferingStopped:
+		CurrentStatus &= ~EMediaStatus::Buffering;
+		break;
 
-	TComPtr<IMFPresentationDescriptor> PresentationDescriptor;
-	{
-		HRESULT Result = MediaSourceObject->CreatePresentationDescriptor(&PresentationDescriptor);
+	case MEConnectEnd:
+		CurrentStatus &= ~EMediaStatus::Connecting;
+		break;
 
-		if (FAILED(Result))
-		{
-			UE_LOG(LogMfMedia, Error, TEXT("Failed to create presentation descriptor: %s"), *MfMedia::ResultToString(Result));
-			return false;
-		}
-	}
+	case MEConnectStart:
+		CurrentStatus |= EMediaStatus::Connecting;
+		EventSink.ReceiveMediaEvent(EMediaEvent::MediaConnecting);
+		break;
 
-	// create source reader
-	HRESULT Result = ::MFCreateSourceReaderFromMediaSource(MediaSourceObject, NULL, &SourceReader);
+	case MESourceCharacteristicsChanged:
+		UpdateCharacteristics();
+		break;
 
-	if (FAILED(Result))
-	{
-		UE_LOG(LogMfMedia, Error, TEXT("Failed to create source reader: %s"), *MfMedia::ResultToString(Result));
-		return false;
+	default:
+		break; // unsupported event
 	}
+}
 
-	// initialize tracks
-	Tracks->Initialize(PresentationDescriptor, MediaSourceObject, SourceReader, Info);
-	MediaEvent.Broadcast(EMediaEvent::TracksChanged);
-
-	// get media duration
-	UINT64 PresentationDuration = 0;
-	PresentationDescriptor->GetUINT64(MF_PD_DURATION, &PresentationDuration);
-	Duration = FTimespan(PresentationDuration);
 
-	// get rate control and support, if available
-	if (SUCCEEDED(SourceReader->GetServiceForStream(MF_SOURCE_READER_MEDIASOURCE, MF_RATE_CONTROL_SERVICE, IID_IMFRateControl, (LPVOID*)&RateControl)) &&
-		SUCCEEDED(SourceReader->GetServiceForStream(MF_SOURCE_READER_MEDIASOURCE, MF_RATE_CONTROL_SERVICE, IID_IMFRateSupport, (LPVOID*)&RateSupport)))
-	{
-		CanScrub = SUCCEEDED(RateSupport->IsRateSupported(TRUE, 0, NULL));
-	}
+void FMfMediaPlayer::ReceiveSourceReaderFlush()
+{
+	Samples->FlushSamples();
+}
 
-	// get capabilities
-	PROPVARIANT Variant;
-	PropVariantInit(&Variant);
 
-	if (SUCCEEDED(SourceReader->GetPresentationAttribute(MF_SOURCE_READER_MEDIASOURCE, MF_SOURCE_READER_MEDIASOURCE_CHARACTERISTICS, &Variant)))
+void FMfMediaPlayer::ReceiveSourceReaderSample(IMFSample* Sample, HRESULT Status, DWORD StreamFlags, DWORD StreamIndex, FTimespan Time)
+{
+	if ((StreamFlags & MF_SOURCE_READERF_ERROR) != 0)
 	{
-		Characteristics = (Variant.vt == VT_UI4) ? Variant.ulVal : 0;
+		SourceReaderError = true;
 	}
 
-	PropVariantClear(&Variant);
-
-	// create session
-	MediaUrl = SourceUrl;
-	MediaSource = MediaSourceObject;
-
-	return true;
+	Tracks->ProcessSample(Sample, Status, StreamFlags, StreamIndex, Time);
 }
 
 
+#undef LOCTEXT_NAMESPACE
+
 #if PLATFORM_WINDOWS
 	#include "HideWindowsPlatformTypes.h"
 #else
diff --git a/Engine/Plugins/Media/MfMedia/Source/MfMedia/Private/Player/MfMediaPlayer.h b/Engine/Plugins/Media/MfMedia/Source/MfMedia/Private/Player/MfMediaPlayer.h
--- a/Engine/Plugins/Media/MfMedia/Source/MfMedia/Private/Player/MfMediaPlayer.h
+++ b/Engine/Plugins/Media/MfMedia/Source/MfMedia/Private/Player/MfMediaPlayer.h
@@ -6,11 +6,16 @@
 
 #if MFMEDIA_SUPPORTED_PLATFORM
 
+#include "Containers/UnrealString.h"
+#include "IMediaCache.h"
 #include "IMediaControls.h"
 #include "IMediaPlayer.h"
-#include "IMfMediaResolverCallbacks.h"
-#include "Containers/Ticker.h"
-#include "Containers/Queue.h"
+#include "IMediaView.h"
+#include "Misc/Timespan.h"
+#include "Templates/SharedPointer.h"
+#include "Windows/COMPointer.h"
+
+#include "IMfMediaSourceReaderSink.h"
 
 #if PLATFORM_WINDOWS
 	#include "WindowsHWrapper.h"
@@ -19,139 +24,164 @@
 	#include "XboxOneAllowPlatformTypes.h"
 #endif
 
-
-class FMfMediaAudioTrack;
-class FMfMediaCaptionTrack;
-class FMfMediaResolver;
+class FMediaSamples;
+class FMfMediaSourceReaderCallback;
 class FMfMediaTracks;
-class FMfMediaVideoTrack;
-class IMediaControls;
-class IMediaOutput;
+class IMediaEventSink;
 
 
 /**
  * Implements a media player using the Windows Media Foundation framework.
  */
 class FMfMediaPlayer
-	: public IMediaControls
-	, public IMediaPlayer
-	, protected IMfMediaResolverCallbacks
+	: public IMediaPlayer
+	, protected IMediaCache
+	, protected IMediaControls
+	, protected IMediaView
+	, protected IMfMediaSourceReaderSink
 {
 public:
 
-	/** Default constructor. */
-	FMfMediaPlayer();
+	/**
+	 * Create and initialize a new instance.
+	 *
+	 * @param InEventSink The object that receives media events from this player.
+	 */
+	FMfMediaPlayer(IMediaEventSink& InEventSink);
 
 	/** Virtual destructor. */
 	virtual ~FMfMediaPlayer();
 
-public:
-
-	//~ IMediaControls interface
-
-	virtual FTimespan GetDuration() const override;
-	virtual float GetRate() const override;
-	virtual EMediaState GetState() const override;
-	virtual TRange<float> GetSupportedRates(EMediaPlaybackDirections Direction, bool Unthinned) const override;
-	virtual FTimespan GetTime() const override;
-	virtual bool IsLooping() const override;
-	virtual bool Seek(const FTimespan& Time) override;
-	virtual bool SetLooping(bool Looping) override;
-	virtual bool SetRate(float Rate) override;
-	virtual bool SupportsRate(float Rate, bool Unthinned) const override;
-	virtual bool SupportsScrubbing() const override;
-	virtual bool SupportsSeeking() const override;
-
 public:
 
 	//~ IMediaPlayer interface
 
 	virtual void Close() override;
+	virtual IMediaCache& GetCache() override;
 	virtual IMediaControls& GetControls() override;
 	virtual FString GetInfo() const override;
 	virtual FName GetName() const override;
-	virtual IMediaOutput& GetOutput() override;
+	virtual IMediaSamples& GetSamples() override;
 	virtual FString GetStats() const override;
 	virtual IMediaTracks& GetTracks() override;
 	virtual FString GetUrl() const override;
-	virtual bool Open(const FString& Url, const IMediaOptions& Options) override;
-	virtual bool Open(const TSharedRef<FArchive, ESPMode::ThreadSafe>& Archive, const FString& OriginalUrl, const IMediaOptions& Options) override;
-	virtual void TickPlayer(float DeltaTime) override;
-	virtual void TickVideo(float DeltaTime) override;
-
-	DECLARE_DERIVED_EVENT(FMfMediaPlayer, IMediaPlayer::FOnMediaEvent, FOnMediaEvent);
-	virtual FOnMediaEvent& OnMediaEvent() override
-	{
-		return MediaEvent;
-	}
+	virtual IMediaView& GetView() override;
+	virtual bool Open(const FString& Url, const IMediaOptions* Options) override;
+	virtual bool Open(const TSharedRef<FArchive, ESPMode::ThreadSafe>& Archive, const FString& OriginalUrl, const IMediaOptions* Options) override;
+	virtual void TickAudio() override;
+	virtual void TickFetch(FTimespan DeltaTime) override;
+	virtual void TickInput(FTimespan DeltaTime) override;
 
 protected:
 
-	//~ IMfMediaResolverCallbacks interface
-
-	virtual void ProcessResolveComplete(TComPtr<IUnknown> SourceObject, FString ResolvedUrl) override;
-	virtual void ProcessResolveFailed(FString FailedUrl) override;
-
-protected:
+	/**
+	 * Commit the specified play position.
+	 *
+	 * @param Time The play position to commit.
+	 */
+	bool CommitTime(FTimespan Time);
 
 	/**
-	 * Finishes opening a new media source.
+	 * Initialize the native AvPlayer instance.
 	 *
-	 * @param SourceObject The media source object.
-	 * @param SourceUrl The original URL of the media source.
+	 * @param Archive The archive being used as a media source (optional).
+	 * @param Url The media URL being opened.
+	 * @param Precache Whether to precache media into RAM if InURL is a local file.
 	 * @return true on success, false otherwise.
 	 */
-	bool FinishOpen(IUnknown* SourceObject, const FString& SourceUrl);
+	bool InitializePlayer(const TSharedPtr<FArchive, ESPMode::ThreadSafe>& Archive, const FString& Url, bool Precache);
 
-private:
+	/** Get the latest characteristics from the current media source. */
+	void UpdateCharacteristics();
+
+protected:
+
+	//~ IMediaControls interface
 
-	/** Caches whether the media supports scrubbing. */
-	bool CanScrub;
+	virtual bool CanControl(EMediaControl Control) const override;
+	virtual FTimespan GetDuration() const override;
+	virtual float GetRate() const override;
+	virtual EMediaState GetState() const override;
+	virtual EMediaStatus GetStatus() const override;
+	virtual TRangeSet<float> GetSupportedRates(EMediaRateThinning Thinning) const override;
+	virtual FTimespan GetTime() const override;
+	virtual bool IsLooping() const override;
+	virtual bool Seek(const FTimespan& Time) override;
+	virtual bool SetLooping(bool Looping) override;
+	virtual bool SetRate(float Rate) override;
+
+protected:
+
+	//~ IMfMediaSourceReaderSink interface
+
+	virtual void ReceiveSourceReaderEvent(MediaEventType Event) override;
+	virtual void ReceiveSourceReaderFlush() override;
+	virtual void ReceiveSourceReaderSample(IMFSample* Sample, HRESULT Status, DWORD StreamFlags, DWORD StreamIndex, FTimespan Time) override;
+
+private:
 
 	/** Cached media characteristics (capabilities). */
 	ULONG Characteristics;
 
+	/** The duration of the currently loaded media. */
+	FTimespan CurrentDuration;
+
 	/** The current playback rate. */
 	float CurrentRate;
 
-	/** Current playback time/position. */
-	FTimespan CurrentTime;
-
-	/** The duration of the currently loaded media. */
-	FTimespan Duration;
-
-	/** Tasks deferred to the game thread. */
-	TQueue<TFunction<void()>> GameThreadTasks;
+	/** The current playback state. */
+	EMediaState CurrentState;
 
-	/** Media information string. */
-	FString Info;
+	/** Current status flags. */
+	EMediaStatus CurrentStatus;
 
-	/** Whether playback should loop to the beginning. */
-	bool Looping;
+	/** Current playback time. */
+	FTimespan CurrentTime;
 
-	/** Event delegate that is invoked when a media event occurred. */
-	FOnMediaEvent MediaEvent;
+	/** The media event handler. */
+	IMediaEventSink& EventSink;
 
-	/** Pointer to the media source object. */
+	/** The currently opened media. */
 	TComPtr<IMFMediaSource> MediaSource;
 
 	/** The URL of the currently opened media. */
 	FString MediaUrl;
 
+	/** If playback just restarted from the Stopped state. */
+	bool PlaybackRestarted;
+
+	/** The presentation descriptor of the currently opened media. */
+	TComPtr<IMFPresentationDescriptor> PresentationDescriptor;
+
 	/** Optional interface for controlling playback rates. */
 	TComPtr<IMFRateControl> RateControl;
 
 	/** Optional interface for querying supported playback rates. */
 	TComPtr<IMFRateSupport> RateSupport;
 
-	/** The media source resolver. */
-	TComPtr<FMfMediaResolver> Resolver;
+	/** Media sample collection. */
+	TSharedPtr<FMediaSamples, ESPMode::ThreadSafe> Samples;
+
+	/** Whether playback should loop to the beginning. */
+	bool ShouldLoop;
 
-	/** The media source reader. */
+	/** The source reader to use. */
 	TComPtr<IMFSourceReader> SourceReader;
 
+	/** Whether an error occurred in the source reader. */
+	bool SourceReaderError;
+
+	/** The source reader callback object.*/
+	TComPtr<FMfMediaSourceReaderCallback> SourceReaderCallback;
+
+	/** The thinned play rates that the current media session supports. */
+	TRangeSet<float> ThinnedRates;
+
 	/** Track collection. */
-	TComPtr<FMfMediaTracks> Tracks;
+	TSharedPtr<FMfMediaTracks, ESPMode::ThreadSafe> Tracks;
+
+	/** The unthinned play rates that the current media session supports. */
+	TRangeSet<float> UnthinnedRates;
 };
 
 
diff --git a/Engine/Plugins/Media/MfMedia/Source/MfMedia/Private/Player/MfMediaSample.h b/Engine/Plugins/Media/MfMedia/Source/MfMedia/Private/Player/MfMediaSample.h
new file mode 100644
--- /dev/null
+++ b/Engine/Plugins/Media/MfMedia/Source/MfMedia/Private/Player/MfMediaSample.h
@@ -0,0 +1,82 @@
+// Copyright 1998-2017 Epic Games, Inc. All Rights Reserved.
+
+#pragma once
+
+#include "MfMediaUtils.h"
+
+
+/**
+ * Abstract base class for MF media samples.
+ */
+class FMfMediaSample
+{
+protected:
+
+	FMfMediaSample()
+		: Duration(FTimespan::Zero())
+		, Time(FTimespan::Zero())
+	{ }
+
+protected:
+
+	/**
+	 * Initialize this sample.
+	 *
+	 * @param InSample The WMF sample.
+	 * @return true on success, false otherwise.
+	 */
+	bool InitializeSample(IMFSample& InSample)
+	{
+		// get sample duration
+		{
+			LONGLONG SampleDuration = 0;
+			HRESULT Result = InSample.GetSampleDuration(&SampleDuration);
+
+			if (FAILED(Result))
+			{
+				UE_LOG(LogMfMedia, VeryVerbose, TEXT("Failed to get sample duration from media sample (%s)"), *MfMedia::ResultToString(Result));
+				return false;
+			}
+
+			Duration = FTimespan(SampleDuration);
+		}
+
+		// get sample flags
+		{
+			HRESULT Result = InSample.GetSampleFlags(&Flags);
+
+			if (FAILED(Result))
+			{
+				UE_LOG(LogMfMedia, VeryVerbose, TEXT("Failed to get sample flags from media sample (%s)"), *MfMedia::ResultToString(Result));
+				return false;
+			}
+		}
+
+		// get sample time
+		{
+			LONGLONG SampleTime = 0;
+			HRESULT Result = InSample.GetSampleTime(&SampleTime);
+
+			if (FAILED(Result))
+			{
+				UE_LOG(LogMfMedia, VeryVerbose, TEXT("Failed to get sample time from media sample (%s)"), *MfMedia::ResultToString(Result));
+				return false;
+			}
+
+			Time = FTimespan(SampleTime);
+		}
+
+		return true;
+	}
+
+protected:
+
+	/** Duration for which the sample is valid. */
+	FTimespan Duration;
+
+	/** Sample flags. */
+	DWORD Flags;
+
+	/** Presentation time for which the sample was generated. */
+	FTimespan Time;
+};
diff --git a/Engine/Plugins/Media/MfMedia/Source/MfMedia/Private/Player/MfMediaTracks.cpp b/Engine/Plugins/Media/MfMedia/Source/MfMedia/Private/Player/MfMediaTracks.cpp
--- a/Engine/Plugins/Media/MfMedia/Source/MfMedia/Private/Player/MfMediaTracks.cpp
+++ b/Engine/Plugins/Media/MfMedia/Source/MfMedia/Private/Player/MfMediaTracks.cpp
@@ -1,9 +1,18 @@
 // Copyright 1998-2017 Epic Games, Inc. All Rights Reserved.
 
 #include "MfMediaTracks.h"
+#include "MfMediaPrivate.h"
 
 #if MFMEDIA_SUPPORTED_PLATFORM
 
+#include "IMediaTextureSample.h"
+#include "Internationalization/Internationalization.h"
+#include "Math/IntPoint.h"
+#include "MediaSamples.h"
+#include "Misc/ScopeLock.h"
+
+#include "MfMediaAudioSample.h"
+#include "MfMediaTextureSample.h"
 #include "MfMediaUtils.h"
 
 #if PLATFORM_WINDOWS
@@ -13,472 +22,541 @@
 	#include "XboxOneAllowPlatformTypes.h"
 #endif
 
+#define MFMEDIATRACKS_TRACE_SAMPLES 0
+#define MFMEDIATRACKS_USE_ASYNCREADER 1
+#define MFMEDIATRACKS_USE_DXVA 0			// not implemented yet
 
-#define LOCTEXT_NAMESPACE "FMfMediaTracks"
+#define LOCTEXT_NAMESPACE "FMfMediaSession"
 
 
-/* Local helpers
- *****************************************************************************/
-
-namespace MfMediaTracks
-{
-	/** How far ahead audio samples should be read to prevent skipping (in seconds). */
-	const float MFPLAYER_AUDIO_READ_AHEAD_SECONDS = 1.0f;
-
-
-	/** Convert a FourCC code to string. */
-	FString FourccToString(unsigned long Fourcc)
-	{
-		return FString::Printf(TEXT("%c%c%c%c"),
-			Fourcc & 0xff,
-			(Fourcc >> 8) & 0xff,
-			(Fourcc >> 16) & 0xff,
-			(Fourcc >> 24) & 0xff
-		);
-	}
-
-
-	/** Convert a GUID to string. */
-	FString GuidToString(const GUID& Guid)
-	{
-		return FString::Printf(TEXT("%08x-%04x-%04x-%08x%08x"),
-			Guid.Data1,
-			Guid.Data2,
-			Guid.Data3,
-			*((unsigned long*)&Guid.Data4[0]),
-			*((unsigned long*)&Guid.Data4[4])
-		);
-	}
-}
-
-
-/* FAndroidMediaTracks structors
+/* FMfMediaSession structors
  *****************************************************************************/
 
 FMfMediaTracks::FMfMediaTracks()
-	: AudioSink(nullptr)
-	, OverlaySink(nullptr)
-	, VideoSink(nullptr)
+	: AudioDone(true)
+	, AudioSamplePending(false)
+	, AudioSamplePool(new FMfMediaAudioSamplePool)
+	, CaptionDone(true)
+	, CaptionSamplePending(false)
+	, LastAudioSampleTime(FTimespan::MinValue())
+	, LastCaptionSampleTime(FTimespan::MinValue())
+	, LastVideoSampleTime(FTimespan::MinValue())
+	, MediaSourceChanged(false)
 	, SelectedAudioTrack(INDEX_NONE)
 	, SelectedCaptionTrack(INDEX_NONE)
 	, SelectedVideoTrack(INDEX_NONE)
-	, AudioDone(true)
-	, CaptionDone(true)
-	, Enabled(false)
+	, SelectionChanged(false)
 	, VideoDone(true)
-	, bIsStarted(false)
-{
-}
+	, VideoSamplePending(false)
+	, VideoSamplePool(new FMfMediaTextureSamplePool)
+{ }
 
 
 FMfMediaTracks::~FMfMediaTracks()
 {
-	FlushRenderingCommands();
-	Reset();
+	Shutdown();
+
+	delete AudioSamplePool;
+	AudioSamplePool = nullptr;
+
+	delete VideoSamplePool;
+	VideoSamplePool = nullptr;
 }
 
 
 /* FMfMediaTracks structors
  *****************************************************************************/
 
-void FMfMediaTracks::Initialize(IMFPresentationDescriptor* InPresentationDescriptor, IMFMediaSource* InMediaSource, IMFSourceReader* InSourceReader, FString& OutInfo)
+void FMfMediaTracks::AppendStats(FString &OutStats) const
 {
 	FScopeLock Lock(&CriticalSection);
 
-	SourceReader = InSourceReader;
+	// audio tracks
+	OutStats += TEXT("Audio Tracks\n");
 
-	// get number of streams
-	DWORD StreamCount = 0;
+	if (AudioTracks.Num() == 0)
+	{
+		OutStats += TEXT("    none\n");
+	}
+	else
 	{
-		if (FAILED(InPresentationDescriptor->GetStreamDescriptorCount(&StreamCount)))
+		for (const FTrack& Track : AudioTracks)
 		{
-			UE_LOG(LogMfMedia, Error, TEXT("Failed to get stream count"));
-
-			return;
+			OutStats += FString::Printf(TEXT("    %s\n"), *Track.DisplayName.ToString());
+			OutStats += TEXT("        Not implemented yet");
 		}
 	}
 
-	// add streams (Media Foundation reports them in reverse order)
-	for (int32 StreamIndex = StreamCount - 1; StreamIndex >= 0; --StreamIndex)
+	// video tracks
+	OutStats += TEXT("Video Tracks\n");
+
+	if (VideoTracks.Num() == 0)
 	{
-		AddStreamToTracks(StreamIndex, InPresentationDescriptor, InMediaSource, OutInfo);
-		OutInfo += TEXT("\n");
+		OutStats += TEXT("    none\n");
+	}
+	else
+	{
+		for (const FTrack& Track : VideoTracks)
+		{
+			OutStats += FString::Printf(TEXT("    %s\n"), *Track.DisplayName.ToString());
+			OutStats += TEXT("        Not implemented yet");
+		}
 	}
-
-	Reinitialize();
 }
 
 
-void FMfMediaTracks::Reinitialize()
+void FMfMediaTracks::ClearFlags()
 {
-	AudioDone = (AudioTracks.Num() == 0);
-	CaptionDone = (CaptionTracks.Num() == 0);
-	VideoDone = (VideoTracks.Num() == 0);
+	FScopeLock Lock(&CriticalSection);
+	
+	MediaSourceChanged = false;
+	SelectionChanged = false;
 }
 
 
-void FMfMediaTracks::Reset()
+FTimespan FMfMediaTracks::GetDuration() const
 {
 	FScopeLock Lock(&CriticalSection);
 
-	SelectedAudioTrack = INDEX_NONE;
-	SelectedCaptionTrack = INDEX_NONE;
-	SelectedVideoTrack = INDEX_NONE;
-
-	if (AudioSink != nullptr)
+	if (PresentationDescriptor == NULL)
 	{
-		AudioSink->ShutdownAudioSink();
+		return FTimespan::Zero();
 	}
 
-	if (OverlaySink != nullptr)
-	{
-		OverlaySink->ShutdownOverlaySink();
-	}
+	UINT64 PresentationDuration = 0;
+	PresentationDescriptor->GetUINT64(MF_PD_DURATION, &PresentationDuration);
 
-	if (VideoSink != nullptr)
-	{
-		VideoSink->ShutdownTextureSink();
-	}
+	return FTimespan(PresentationDuration);
+}
 
-	AudioTracks.Empty();
-	CaptionTracks.Empty();
-	VideoTracks.Empty();
 
-	AudioDone = true;
-	CaptionDone = true;
-	VideoDone = true;
-	Enabled = false;
+void FMfMediaTracks::GetFlags(bool& OutMediaSourceChanged, bool& OutSelectionChanged) const
+{
+	FScopeLock Lock(&CriticalSection);
 
-	SourceReader = nullptr;
+	OutMediaSourceChanged = MediaSourceChanged;
+	OutSelectionChanged = SelectionChanged;
 }
 
 
-void FMfMediaTracks::SetEnabled(bool InEnabled)
+IMFMediaSource* FMfMediaTracks::GetMediaSource()
 {
 	FScopeLock Lock(&CriticalSection);
-	Enabled = InEnabled;
+	return MediaSource;
+}
 
-	if (AudioSink != nullptr)
-	{
-		if (InEnabled)
-		{
-			AudioSink->ResumeAudioSink();
-		}
-		else
-		{
-			AudioSink->PauseAudioSink();
-		}
-	}
+
+IMFSourceReader* FMfMediaTracks::GetSourceReader()
+{
+	FScopeLock Lock(&CriticalSection);
+	return SourceReader;
 }
 
 
-void FMfMediaTracks::Tick(float DeltaTime)
+void FMfMediaTracks::Initialize(IMFMediaSource* InMediaSource, IMFSourceReaderCallback* InSourceReaderCallback, const TSharedRef<FMediaSamples, ESPMode::ThreadSafe>& InSamples)
 {
+	Shutdown();
+
 	FScopeLock Lock(&CriticalSection);
 
-	if (!Enabled || (SourceReader == NULL))
+	MediaSourceChanged = true;
+
+	if (InMediaSource == NULL)
 	{
 		return;
 	}
 
-	// Clear out the delta time on the first tick of playback.
-	// This guarantees we will start playback at the first frame.
-	if (!bIsStarted)
+	// create presentation descriptor
+	TComPtr<IMFPresentationDescriptor> NewPresentationDescriptor;
 	{
-		DeltaTime = 0.0f;
-		bIsStarted = true;
+		HRESULT Result = InMediaSource->CreatePresentationDescriptor(&NewPresentationDescriptor);
+
+		if (FAILED(Result))
+		{
+			UE_LOG(LogMfMedia, Error, TEXT("Failed to create presentation descriptor: %s"), *MfMedia::ResultToString(Result));
+			return;
+		}
 	}
 
-	// Check if new video sample(s) required
-	if (VideoTracks.IsValidIndex(SelectedVideoTrack) && !VideoDone)
+	// create source reader
+	TComPtr<IMFAttributes> Attributes;
 	{
-		FVideoTrack& VideoTrack = VideoTracks[SelectedVideoTrack];
+		if (FAILED(MFCreateAttributes(&Attributes, 1)))
+		{
+			UE_LOG(LogMfMedia, Error, TEXT("Failed to create source reader attributes"));
+			return;
+		}
 
-		// Decrease the time remaining on the currently displayed sample by the passage of time
-		VideoTrack.SecondsUntilNextSample -= DeltaTime;
+		if (FAILED(Attributes->SetUINT32(MF_READWRITE_ENABLE_HARDWARE_TRANSFORMS, TRUE)))
+		{
+			UE_LOG(LogMfMedia, Warning, TEXT("Failed to set one or more source reader attributes"));
+		}
 
-		// Read through samples until we catch up
-		IMFSample* Sample = NULL;
-		LONGLONG Timestamp = 0;
-		while (VideoTrack.SecondsUntilNextSample <= 0)
+#if MFMEDIATRACKS_USE_DXVA
+		if (FAILED(Attributes->SetUINT32(MF_SOURCE_READER_D3D_MANAGER, D3dDeviceManager)) ||
+			FAILED(Attributes->SetUINT32(MF_SOURCE_READER_DISABLE_DXVA, FALSE)))
 		{
-			SAFE_RELEASE(Sample);
+			UE_LOG(LogMfMedia, Warning, TEXT("Failed to set DXVA source reader attributes"));
+		}
+#endif
 
-			DWORD StreamFlags = 0;
-			HRESULT Result = SourceReader->ReadSample(VideoTrack.StreamIndex, 0, NULL, &StreamFlags, &Timestamp, &Sample);
-			if (FAILED(Result))
-			{
-				UE_LOG(LogMfMedia, Warning, TEXT("Messed up (%s)"), *MfMedia::ResultToString(Result));
-				Sample = NULL;
-				break;
-			}
-			if (StreamFlags & MF_SOURCE_READERF_ENDOFSTREAM)
-			{
-				VideoDone = true;
-				SAFE_RELEASE(Sample);
-				break;
-			}
-			if (StreamFlags & MF_SOURCE_READERF_NATIVEMEDIATYPECHANGED)
-			{
-				//@TODO
-			}
-			if (Sample != NULL)
-			{
-				LONGLONG SampleDuration;
-				Sample->GetSampleDuration(&SampleDuration);
-				FTimespan Time(SampleDuration);
-				VideoTrack.SecondsUntilNextSample += Time.GetTotalSeconds();
-			}
+#if MFMEDIATRACKS_USE_ASYNCREADER
+		if (FAILED(Attributes->SetUnknown(MF_SOURCE_READER_ASYNC_CALLBACK, InSourceReaderCallback)))
+		{
+			UE_LOG(LogMfMedia, Warning, TEXT("Failed to set async callback source reader attribute"));
 		}
+#endif
+	}
+
+	TComPtr<IMFSourceReader> NewSourceReader;
+	{
+		HRESULT Result = ::MFCreateSourceReaderFromMediaSource(InMediaSource, Attributes, &NewSourceReader);
 
-		// Process sample into content
-		if (Sample != NULL)
+		if (FAILED(Result))
 		{
-			if (VideoSink != nullptr)
-			{
-				// get buffer data
-				TComPtr<IMFMediaBuffer> Buffer;
-				if (SUCCEEDED(Sample->ConvertToContiguousBuffer(&Buffer)))
-				{
-					DWORD Length;
-					Buffer->GetCurrentLength(&Length);
-
-					// The video sink needs to be updated on the render thread.
-					// The render thread takes ownership of the sample and is responsible for releasing it in this case.
-					// ReadSample gives us samples from a queue provided by the decoder. It's ok for us to hold onto
-					// this sample even if we ReadSample again before the render thread releases this sample.
-					ENQUEUE_UNIQUE_RENDER_COMMAND_FOURPARAMETER(
-						FUpdateMediaTextureCommand,
-						TComPtr<IMFMediaBuffer>, Buffer, Buffer,
-						LONGLONG, Timestamp, Timestamp,
-						IMediaTextureSink*, VideoSink, VideoSink,
-						IMFSample*, Sample, Sample,
-						{
-							uint8* Data = nullptr;
-							if (SUCCEEDED(Buffer->Lock(&Data, nullptr, nullptr)))
-							{
-								VideoSink->UpdateTextureSinkBuffer(Data);
-								VideoSink->DisplayTextureSinkBuffer(FTimespan(Timestamp));
-								Buffer->Unlock();
-								SAFE_RELEASE(Sample);
-							}
-						});
-				}
-				else
-				{
-					SAFE_RELEASE(Sample);
-				}
-			}
-			else
-			{
-				SAFE_RELEASE(Sample);
-			}
+			UE_LOG(LogMfMedia, Error, TEXT("Failed to create source reader: %s"), *MfMedia::ResultToString(Result));
+			return;
 		}
 	}
 
-	// Check if new audio sample(s) required
-	if (AudioTracks.IsValidIndex(SelectedAudioTrack) && !AudioDone)
+	// get number of streams
+	DWORD StreamCount = 0;
 	{
-		FAudioTrack& AudioTrack = AudioTracks[SelectedAudioTrack];
+		HRESULT Result = NewPresentationDescriptor->GetStreamDescriptorCount(&StreamCount);
+
+		if (FAILED(Result))
+		{
+			UE_LOG(LogMfMedia, Error, TEXT("Failed to get stream count (%s)"), *MfMedia::ResultToString(Result));
+			return;
+		}
 
-		// Decrease the time remaining on the currently queued sample by the passage of time
-		AudioTrack.SecondsUntilNextSample -= DeltaTime;
+		UE_LOG(LogMfMedia, Verbose, TEXT("Tracks %llx: Found %i streams"), this, StreamCount);
+	}
+
+	// initialization successful
+	MediaSource = InMediaSource;
+	PresentationDescriptor = NewPresentationDescriptor;
+	Samples = InSamples;
+	SourceReader = NewSourceReader;
 
-		// Read through samples until we catch up
-		while (AudioTrack.SecondsUntilNextSample <= MfMediaTracks::MFPLAYER_AUDIO_READ_AHEAD_SECONDS)
+	// add streams (Media Foundation reports them in reverse order)
+	for (int32 StreamIndex = StreamCount - 1; StreamIndex >= 0; --StreamIndex)
+	{
+		AddStreamToTracks(StreamIndex, Info);
+		Info += TEXT("\n");
+	}
+}
+
+
+void FMfMediaTracks::ProcessSample(IMFSample* Sample, HRESULT Status, DWORD StreamFlags, DWORD StreamIndex, FTimespan Time)
+{
+	FScopeLock Lock(&CriticalSection);
+
+	if (!Samples.IsValid() || FAILED(Status))
+	{
+		return;
+	}
+
+	if ((StreamFlags & MF_SOURCE_READERF_NATIVEMEDIATYPECHANGED) != 0)
+	{
+		// @todo gmp: MF3.0 re-initialize source reader
+	}
+
+	// process audio sample
+	if (SelectedAudioTrack != INDEX_NONE)
+	{
+		const FTrack& Track = AudioTracks[SelectedAudioTrack];
+
+		if (Track.StreamIndex == StreamIndex)
 		{
-			DWORD StreamFlags = 0;
-			LONGLONG Timestamp = 0;
-			IMFSample* Sample = NULL;
-			if (FAILED(SourceReader->ReadSample(AudioTrack.StreamIndex, 0, NULL, &StreamFlags, &Timestamp, &Sample)))
-			{
-				break;
-			}
-			if (StreamFlags & MF_SOURCE_READERF_ENDOFSTREAM)
+			if (((StreamFlags & MF_SOURCE_READERF_ENDOFSTREAM) != 0) || ((StreamFlags & MF_SOURCE_READERF_ERROR) != 0))
 			{
+				UE_LOG(LogMfMedia, Verbose, TEXT("Tracks %llx: Audio done"), this);
+
 				AudioDone = true;
-				SAFE_RELEASE(Sample);
-				break;
 			}
-			if (Sample != NULL)
+
+			if ((Track.SelectedFormat != INDEX_NONE) && (Sample != NULL))
 			{
-				LONGLONG SampleDuration;
-				Sample->GetSampleDuration(&SampleDuration);
-				FTimespan Time(SampleDuration);
-				AudioTrack.SecondsUntilNextSample += Time.GetTotalSeconds();
+				const FFormat& Format = Track.Formats[Track.SelectedFormat];
+				check(Format.OutputType.IsValid());
+
+				const TSharedRef<FMfMediaAudioSample, ESPMode::ThreadSafe> AudioSample = AudioSamplePool->AcquireShared();
 
-				// Process sample into content
-				if (AudioSink != nullptr)
+				if (AudioSample->Initialize(*Format.OutputType, *Sample, Format.Audio.NumChannels, Format.Audio.SampleRate))
 				{
-					// get buffer data
-					TComPtr<IMFMediaBuffer> Buffer;
-					if (SUCCEEDED(Sample->ConvertToContiguousBuffer(&Buffer)))
-					{
-						DWORD Length = 0;
-						if (SUCCEEDED(Buffer->GetCurrentLength(&Length)))
-						{
-							uint8* Data = nullptr;
-							if (SUCCEEDED(Buffer->Lock(&Data, nullptr, nullptr)))
-							{
-								AudioSink->PlayAudioSink(Data, Length, FTimespan(Timestamp));
-								Buffer->Unlock();
-							}
-						}
-					}
-				}
+					Samples->AddAudio(AudioSample);
+					LastAudioSampleTime = AudioSample->GetTime();
 
-				SAFE_RELEASE(Sample);
+					#if MFMEDIATRACKS_TRACE_SAMPLES
+						UE_LOG(LogMfMedia, VeryVerbose, TEXT("Tracks %llx: Audio sample processed: %s"), this, *LastAudioSampleTime.ToString());
+					#endif
+				}
 			}
+
+			AudioSamplePending = false;
+			UpdateAudio();
+
+			return;
 		}
 	}
 
-	// Check if new caption sample(s) required
-	if (CaptionTracks.IsValidIndex(SelectedCaptionTrack) && !CaptionDone)
+	// process caption sample
+	if (SelectedCaptionTrack != INDEX_NONE)
 	{
-		FCaptionTrack& CaptionTrack = CaptionTracks[SelectedCaptionTrack];
+		const FTrack& Track = CaptionTracks[SelectedCaptionTrack];
 
-		// Decrease the time remaining on the currently displayed sample by the passage of time
-		CaptionTrack.SecondsUntilNextSample -= DeltaTime;
-
-		// Read through samples until we catch up
-		IMFSample* Sample = NULL;
-		LONGLONG Duration = 0;
-		LONGLONG Timestamp = 0;
-		while (CaptionTrack.SecondsUntilNextSample <= 0)
+		if (Track.StreamIndex == StreamIndex)
 		{
-			SAFE_RELEASE(Sample);
-
-			DWORD StreamFlags = 0;
-			if (FAILED(SourceReader->ReadSample(CaptionTrack.StreamIndex, 0, NULL, &StreamFlags, &Timestamp, &Sample)))
-			{
-				Sample = NULL;
-				break;
-			}
-			if (StreamFlags & MF_SOURCE_READERF_ENDOFSTREAM)
+			if (((StreamFlags & MF_SOURCE_READERF_ENDOFSTREAM) != 0) || ((StreamFlags & MF_SOURCE_READERF_ERROR) != 0))
 			{
+				UE_LOG(LogMfMedia, Verbose, TEXT("Tracks %llx: Caption done"), this);
+
 				CaptionDone = true;
-				SAFE_RELEASE(Sample);
-				break;
 			}
-			if (Sample != NULL)
+
+			if ((Track.SelectedFormat != INDEX_NONE) && (Sample != NULL))
 			{
-				Sample->GetSampleDuration(&Duration);
-				FTimespan Time(Duration);
-				CaptionTrack.SecondsUntilNextSample += Time.GetTotalSeconds();
+				const FFormat& Format = Track.Formats[Track.SelectedFormat];
+				check(Format.OutputType.IsValid());
+/*
+				const auto CaptionSample = MakeShared<FMfMediaOverlaySample, ESPMode::ThreadSafe>();
+
+				if (CaptionSample->Initialize(*Format.OutputType, *Sample))
+				{
+					Samples->AddCaption(CaptionSample);
+					LastCaptionSampleTime = CaptionSample->GetTime();
+*/
+					LONGLONG SampleTime = 0;
+					Sample->GetSampleTime(&SampleTime);
+					LastCaptionSampleTime = FTimespan(SampleTime);
+
+					#if MFMEDIATRACKS_TRACE_SAMPLES
+						UE_LOG(LogMfMedia, VeryVerbose, TEXT("Tracks %llx: Caption sample processed: %s"), this, *LastCaptionSampleTime.ToString());
+					#endif
+//				}
 			}
+
+			CaptionSamplePending = false;
+			UpdateCaptions();
 		}
+	}
+	
+	// process video sample
+	if (SelectedVideoTrack != INDEX_NONE)
+	{
+		const FTrack& Track = VideoTracks[SelectedVideoTrack];
 
-		// Process sample into content
-		if (Sample != NULL)
+		if (Track.StreamIndex == StreamIndex)
 		{
-			if (OverlaySink != nullptr)
+			if (((StreamFlags & MF_SOURCE_READERF_ENDOFSTREAM) != 0) || ((StreamFlags & MF_SOURCE_READERF_ERROR) != 0))
+			{
+				UE_LOG(LogMfMedia, Verbose, TEXT("Tracks %llx: Video done"), this);
+
+				VideoDone = true;
+			}
+
+			if ((Track.SelectedFormat != INDEX_NONE) && (Sample != NULL))
 			{
-				// get buffer data
-				TComPtr<IMFMediaBuffer> Buffer;
-				if (SUCCEEDED(Sample->ConvertToContiguousBuffer(&Buffer)))
+				const FFormat& Format = Track.Formats[Track.SelectedFormat];
+				check(Format.OutputType.IsValid());
+
+				const TSharedRef<FMfMediaTextureSample, ESPMode::ThreadSafe> VideoSample = VideoSamplePool->AcquireShared();
+
+				if (VideoSample->Initialize(*Format.OutputType, *Sample, Format.Video.BufferDim, Format.Video.BufferStride, Format.Video.OutputDim, true))
 				{
-					uint8* Data = nullptr;
-					if (SUCCEEDED(Buffer->Lock(&Data, nullptr, nullptr)))
-					{
-						OverlaySink->AddOverlaySinkText(FText::FromString((TCHAR*)Data), EMediaOverlayType::Caption, FTimespan(Timestamp), FTimespan(Duration), TOptional<FVector2D>());
-						Buffer->Unlock();
-					}
+					Samples->AddVideo(VideoSample);
+					LastVideoSampleTime = VideoSample->GetTime();
+
+					#if MFMEDIATRACKS_TRACE_SAMPLES
+						UE_LOG(LogMfMedia, VeryVerbose, TEXT("Tracks %llx: Video sample processed: %s"), this, *LastVideoSampleTime.ToString());
+					#endif
 				}
 			}
-			SAFE_RELEASE(Sample);
+
+			VideoSamplePending = false;
+			UpdateVideo();
+
+			return;
 		}
 	}
 }
 
 
-/* IMediaOutput interface
- *****************************************************************************/
-
-void FMfMediaTracks::SetAudioSink(IMediaAudioSink* Sink)
+void FMfMediaTracks::Restart()
 {
-	if (Sink != AudioSink)
-	{
-		FScopeLock Lock(&CriticalSection);
+	UE_LOG(LogMfMedia, Verbose, TEXT("Tracks %llx: Restarting sample processing"), this);
 
-		if (AudioSink != nullptr)
-		{
-			AudioSink->ShutdownAudioSink();
-		}
+	FScopeLock Lock(&CriticalSection);
 
-		AudioSink = Sink;	
-		InitializeAudioSink();
+	if (SourceReader.IsValid())
+	{
+//		SourceReader->Flush(MF_SOURCE_READER_ALL_STREAMS);
 	}
+
+	AudioDone = (AudioTracks.Num() == 0);
+	CaptionDone = (CaptionTracks.Num() == 0);
+	VideoDone = (VideoTracks.Num() == 0);
+
+	AudioSampleRange = TRange<FTimespan>::Empty();
+	CaptionSampleRange = TRange<FTimespan>::Empty();
+	VideoSampleRange = TRange<FTimespan>::Empty();
+
+	AudioSamplePending = false;
+	CaptionSamplePending = false;
+	VideoSamplePending = false;
+
+	LastAudioSampleTime = FTimespan::MinValue();
+	LastCaptionSampleTime = FTimespan::MinValue();
+	LastVideoSampleTime = FTimespan::MinValue();
 }
 
 
-void FMfMediaTracks::SetMetadataSink(IMediaBinarySink* Sink)
+void FMfMediaTracks::Shutdown()
 {
-	// not implemented yet
+	FScopeLock Lock(&CriticalSection);
+
+	AudioSamplePool->Reset();
+	VideoSamplePool->Reset();
+
+	SelectedAudioTrack = INDEX_NONE;
+	SelectedCaptionTrack = INDEX_NONE;
+	SelectedVideoTrack = INDEX_NONE;
+
+	AudioTracks.Empty();
+	CaptionTracks.Empty();
+	VideoTracks.Empty();
+
+	Info.Empty();
+
+	if (MediaSource != NULL)
+	{
+		MediaSource->Shutdown();
+		MediaSource = NULL;
+	}
+
+	AudioDone = true;
+	CaptionDone = true;
+	VideoDone = true;
+
+	PresentationDescriptor.Reset();
+	Samples.Reset();
+	SourceReader.Reset();
+
+	MediaSourceChanged = false;
 }
 
 
-void FMfMediaTracks::SetOverlaySink(IMediaOverlaySink* Sink)
+void FMfMediaTracks::TickAudio(float Rate, FTimespan Time)
 {
-	if (Sink != OverlaySink)
+	FScopeLock Lock(&CriticalSection);
+
+	if ((Rate <= 0.0f) || (Rate > 2.0f))
 	{
-		FScopeLock Lock(&CriticalSection);
+		return; // no audio in reverse or very fast forward
+	}
 
-		if (OverlaySink != nullptr)
-		{
-			OverlaySink->ShutdownOverlaySink();
-		}
+	AudioSampleRange = TRange<FTimespan>::AtMost(Time + FTimespan::FromSeconds(Rate));
 
-		OverlaySink = Sink;
-		InitializeOverlaySink();
-	}
+	UpdateAudio();
 }
 
 
-void FMfMediaTracks::SetVideoSink(IMediaTextureSink* Sink)
+void FMfMediaTracks::TickInput(float Rate, FTimespan Time, FTimespan DeltaTime)
 {
-	if (Sink != VideoSink)
+	FScopeLock Lock(&CriticalSection);
+
+	if (Rate > 0.0f)
+	{
+		CaptionSampleRange = TRange<FTimespan>::AtMost(Time);
+	}
+	else if (Rate < 0.0f)
 	{
-		FScopeLock Lock(&CriticalSection);
+		CaptionSampleRange = TRange<FTimespan>::AtLeast(Time);
+	}
 
-		if (VideoSink != nullptr)
-		{
-			VideoSink->ShutdownTextureSink();
-		}
+	VideoSampleRange = CaptionSampleRange;
 
-		VideoSink = Sink;
-		InitializeVideoSink();
-	}
+	UpdateCaptions();
+	UpdateVideo();
 }
 
 
 /* IMediaTracks interface
  *****************************************************************************/
 
-uint32 FMfMediaTracks::GetAudioTrackChannels(int32 TrackIndex) const
+bool FMfMediaTracks::GetAudioTrackFormat(int32 TrackIndex, int32 FormatIndex, FMediaAudioTrackFormat& OutFormat) const
 {
-	return AudioTracks.IsValidIndex(TrackIndex) ? AudioTracks[TrackIndex].NumChannels : 0;
-}
+	FScopeLock Lock(&CriticalSection);
 
+	const FFormat* Format = GetAudioFormat(TrackIndex, FormatIndex);
 
-uint32 FMfMediaTracks::GetAudioTrackSampleRate(int32 TrackIndex) const
-{
-	return AudioTracks.IsValidIndex(TrackIndex) ? AudioTracks[TrackIndex].SampleRate : 0;
+	if (Format == nullptr)
+	{
+		return false;
+	}
+
+	OutFormat.BitsPerSample = Format->Audio.BitsPerSample;
+	OutFormat.NumChannels = Format->Audio.NumChannels;
+	OutFormat.SampleRate = Format->Audio.SampleRate;
+	OutFormat.TypeName = Format->TypeName;
+
+	return true;
 }
 
 
 int32 FMfMediaTracks::GetNumTracks(EMediaTrackType TrackType) const
 {
+	FScopeLock Lock(&CriticalSection);
+
 	switch (TrackType)
 	{
 	case EMediaTrackType::Audio:
 		return AudioTracks.Num();
+
 	case EMediaTrackType::Caption:
 		return CaptionTracks.Num();
+
 	case EMediaTrackType::Video:
 		return VideoTracks.Num();
+
 	default:
-		return 0;
+		break; // unsupported track type
 	}
+
+	return 0;
+}
+
+
+int32 FMfMediaTracks::GetNumTrackFormats(EMediaTrackType TrackType, int32 TrackIndex) const
+{
+	FScopeLock Lock(&CriticalSection);
+
+	switch (TrackType)
+	{
+	case EMediaTrackType::Audio:
+		if (AudioTracks.IsValidIndex(TrackIndex))
+		{
+			return AudioTracks[TrackIndex].Formats.Num();
+		}
+
+	case EMediaTrackType::Caption:
+		if (CaptionTracks.IsValidIndex(TrackIndex))
+		{
+			return 1;
+		}
+
+	case EMediaTrackType::Video:
+		if (VideoTracks.IsValidIndex(TrackIndex))
+		{
+			return VideoTracks[TrackIndex].Formats.Num();
+		}
+
+	default:
+		break; // unsupported track type
+	}
+
+	return 0;
 }
 
 
@@ -488,18 +566,25 @@ int32 FMfMediaTracks::GetSelectedTrack(EMediaTrackType TrackType) const
 	{
 	case EMediaTrackType::Audio:
 		return SelectedAudioTrack;
+
 	case EMediaTrackType::Caption:
 		return SelectedCaptionTrack;
+
 	case EMediaTrackType::Video:
 		return SelectedVideoTrack;
+
 	default:
-		return INDEX_NONE;
+		break; // unsupported track type
 	}
+
+	return INDEX_NONE;
 }
 
 
 FText FMfMediaTracks::GetTrackDisplayName(EMediaTrackType TrackType, int32 TrackIndex) const
 {
+	FScopeLock Lock(&CriticalSection);
+
 	switch (TrackType)
 	{
 	case EMediaTrackType::Audio:
@@ -524,15 +609,26 @@ FText FMfMediaTracks::GetTrackDisplayName(EMediaTrackType TrackType, int32 Track
 		break;
 
 	default:
-		break;
+		break; // unsupported track type
 	}
 
 	return FText::GetEmpty();
 }
 
 
+int32 FMfMediaTracks::GetTrackFormat(EMediaTrackType TrackType, int32 TrackIndex) const
+{
+	FScopeLock Lock(&CriticalSection);
+
+	const FTrack* Track = GetTrack(TrackType, TrackIndex);
+	return (Track != nullptr) ? Track->SelectedFormat : INDEX_NONE;
+}
+
+
 FString FMfMediaTracks::GetTrackLanguage(EMediaTrackType TrackType, int32 TrackIndex) const
 {
+	FScopeLock Lock(&CriticalSection);
+
 	switch (TrackType)
 	{
 	case EMediaTrackType::Audio:
@@ -557,7 +653,7 @@ FString FMfMediaTracks::GetTrackLanguage(EMediaTrackType TrackType, int32 TrackI
 		break;
 
 	default:
-		break;
+		break; // unsupported track type
 	}
 
 	return FString();
@@ -566,6 +662,8 @@ FString FMfMediaTracks::GetTrackLanguage(EMediaTrackType TrackType, int32 TrackI
 
 FString FMfMediaTracks::GetTrackName(EMediaTrackType TrackType, int32 TrackIndex) const
 {
+	FScopeLock Lock(&CriticalSection);
+
 	switch (TrackType)
 	{
 	case EMediaTrackType::Audio:
@@ -590,161 +688,241 @@ FString FMfMediaTracks::GetTrackName(EMediaTrackType TrackType, int32 TrackIndex
 		break;
 
 	default:
-		break;
+		break; // unsupported track type
 	}
 
 	return FString();
 }
 
 
-uint32 FMfMediaTracks::GetVideoTrackBitRate(int32 TrackIndex) const
+bool FMfMediaTracks::GetVideoTrackFormat(int32 TrackIndex, int32 FormatIndex, FMediaVideoTrackFormat& OutFormat) const
 {
-	return VideoTracks.IsValidIndex(TrackIndex) ? VideoTracks[TrackIndex].BitRate : 0;
-}
+	FScopeLock Lock(&CriticalSection);
 
+	const FFormat* Format = GetVideoFormat(TrackIndex, FormatIndex);
 
-FIntPoint FMfMediaTracks::GetVideoTrackDimensions(int32 TrackIndex) const
-{
-	return VideoTracks.IsValidIndex(TrackIndex) ? VideoTracks[TrackIndex].OutputDim : FIntPoint::ZeroValue;
-}
+	if (Format == nullptr)
+	{
+		return false;
+	}
 
+	OutFormat.Dim = Format->Video.OutputDim;
+	OutFormat.FrameRate = Format->Video.FrameRate;
+	OutFormat.FrameRates = Format->Video.FrameRates;
+	OutFormat.TypeName = Format->TypeName;
 
-float FMfMediaTracks::GetVideoTrackFrameRate(int32 TrackIndex) const
-{
-	return VideoTracks.IsValidIndex(TrackIndex) ? VideoTracks[TrackIndex].FrameRate : 0.0f;
+	return true;
 }
 
 
 bool FMfMediaTracks::SelectTrack(EMediaTrackType TrackType, int32 TrackIndex)
 {
+	if (SourceReader == nullptr)
+	{
+		return false; // not initialized
+	}
+
 	FScopeLock Lock(&CriticalSection);
 
+	int32 OldStreamIndex = INDEX_NONE;
+	int32 NewStreamIndex = INDEX_NONE;
+	int32* SelectedTrack = nullptr;
+
 	switch (TrackType)
 	{
 	case EMediaTrackType::Audio:
-		if ((TrackIndex == INDEX_NONE) || AudioTracks.IsValidIndex(TrackIndex))
+		if (TrackIndex == SelectedAudioTrack)
 		{
-			if (TrackIndex != SelectedAudioTrack)
-			{
-				if (AudioTracks.IsValidIndex(SelectedAudioTrack))
-				{
-					SourceReader->SetStreamSelection(AudioTracks[SelectedAudioTrack].StreamIndex, FALSE);
-				}
-				if (AudioTracks.IsValidIndex(TrackIndex))
-				{
-					SourceReader->SetStreamSelection(AudioTracks[TrackIndex].StreamIndex, TRUE);
-				}
-				SelectedAudioTrack = TrackIndex;
-				InitializeAudioSink();
-			}
+			return true; // already selected
+		}
 
-			return true;
+		if ((TrackIndex != INDEX_NONE) && !AudioTracks.IsValidIndex(TrackIndex))
+		{
+			return false; // invalid track
 		}
-		break;
 
-	case EMediaTrackType::Caption:
-		if ((TrackIndex == INDEX_NONE) || CaptionTracks.IsValidIndex(TrackIndex))
+		UE_LOG(LogMfMedia, Verbose, TEXT("Tracks %llx: Selecting audio track %i instead of %i (%i tracks)"), this, TrackIndex, SelectedAudioTrack, AudioTracks.Num());
+
+		if (AudioTracks.IsValidIndex(SelectedAudioTrack))
 		{
-			if (TrackIndex != SelectedCaptionTrack)
-			{
-				if (CaptionTracks.IsValidIndex(SelectedCaptionTrack))
-				{
-					SourceReader->SetStreamSelection(CaptionTracks[SelectedCaptionTrack].StreamIndex, FALSE);
-				}
-				if (CaptionTracks.IsValidIndex(TrackIndex))
-				{
-					SourceReader->SetStreamSelection(CaptionTracks[TrackIndex].StreamIndex, TRUE);
-				}
-				SelectedCaptionTrack = TrackIndex;
-				InitializeOverlaySink();
-			}
+			OldStreamIndex = AudioTracks[SelectedAudioTrack].StreamIndex;
+		}
 
-			return true;
+		if (TrackIndex != INDEX_NONE)
+		{
+			NewStreamIndex = AudioTracks[TrackIndex].StreamIndex;
 		}
+
+		SelectedTrack = &SelectedAudioTrack;
 		break;
 
 	case EMediaTrackType::Video:
-		if ((TrackIndex == INDEX_NONE) || VideoTracks.IsValidIndex(TrackIndex))
+		if (TrackIndex == SelectedVideoTrack)
 		{
-			if (TrackIndex != SelectedVideoTrack)
-			{
-				if (VideoTracks.IsValidIndex(SelectedVideoTrack))
-				{
-					SourceReader->SetStreamSelection(VideoTracks[SelectedVideoTrack].StreamIndex, FALSE);
-				}
-				if (VideoTracks.IsValidIndex(TrackIndex))
-				{
-					SourceReader->SetStreamSelection(VideoTracks[TrackIndex].StreamIndex, TRUE);
-				}
-				SelectedVideoTrack = TrackIndex;
-				InitializeVideoSink();
-			}
+			return true; // already selected
+		}
+
+		if ((TrackIndex != INDEX_NONE) && !VideoTracks.IsValidIndex(TrackIndex))
+		{
+			return false; // invalid track
+		}
+
+		UE_LOG(LogMfMedia, Verbose, TEXT("Tracks %llx: Selecting video track %i instead of %i (%i tracks)"), this, TrackIndex, SelectedVideoTrack, VideoTracks.Num());
 
-			return true;
+		if (VideoTracks.IsValidIndex(SelectedVideoTrack))
+		{
+			OldStreamIndex = VideoTracks[SelectedVideoTrack].StreamIndex;
+		}
+
+		if (TrackIndex != INDEX_NONE)
+		{
+			NewStreamIndex = VideoTracks[TrackIndex].StreamIndex;
 		}
+
+		SelectedTrack = &SelectedVideoTrack;
 		break;
 
 	default:
-		break;
+		return false;
 	}
 
-	return false;
-}
+	check(SelectedTrack != nullptr);
 
+	// apply changes to source reader
+	if (OldStreamIndex != INDEX_NONE)
+	{
+		HRESULT Result = SourceReader->SetStreamSelection(OldStreamIndex, FALSE);
 
-/* IUnknown interface
- *****************************************************************************/
+		if (FAILED(Result))
+		{
+			UE_LOG(LogMfMedia, Warning, TEXT("Failed to deselect stream %i: %s"), OldStreamIndex, *MfMedia::ResultToString(Result));
+			return false;
+		}
 
-STDMETHODIMP FMfMediaTracks::QueryInterface(REFIID RefID, void** Object)
-{
-	if (Object == NULL)
-	{
-		return E_INVALIDARG;
+		Result = SourceReader->Flush(OldStreamIndex);
+
+		if (FAILED(Result))
+		{
+			UE_LOG(LogMfMedia, Warning, TEXT("Failed to flush deselected stream %i: %s"), OldStreamIndex, *MfMedia::ResultToString(Result));
+		}
+
+		UE_LOG(LogMfMedia, Verbose, TEXT("Tracks %llx: Disabled stream %i"), this, OldStreamIndex);
+
+		SelectionChanged = true;
 	}
 
-	if ((RefID == IID_IUnknown) || (RefID == IID_IMFSourceReaderCallback))
+	if (NewStreamIndex != INDEX_NONE)
 	{
-		*Object = (LPVOID)this;
-		AddRef();
+		const HRESULT Result = SourceReader->SetStreamSelection(NewStreamIndex, TRUE);
+
+		if (FAILED(Result))
+		{
+			UE_LOG(LogMfMedia, Warning, TEXT("Failed to select stream %i: %s"), NewStreamIndex, *MfMedia::ResultToString(Result));
+			return false;
+		}
 
-		return NOERROR;
+		UE_LOG(LogMfMedia, Verbose, TEXT("Tracks %llx: Enabled stream %i"), this, NewStreamIndex);
+
+		SelectionChanged = true;
 	}
 
-	*Object = NULL;
+	*SelectedTrack = TrackIndex;
 
-	return E_NOINTERFACE;
+	return true;
 }
 
 
-STDMETHODIMP_(ULONG) FMfMediaTracks::AddRef()
+bool FMfMediaTracks::SetTrackFormat(EMediaTrackType TrackType, int32 TrackIndex, int32 FormatIndex)
 {
-	return FPlatformAtomics::InterlockedIncrement(&RefCount);
-}
+	FScopeLock Lock(&CriticalSection);
 
+	FTrack* Track = nullptr;
 
-STDMETHODIMP_(ULONG) FMfMediaTracks::Release()
-{
-	int32 CurrentRefCount = FPlatformAtomics::InterlockedDecrement(&RefCount);
+	// get track
+	switch (TrackType)
+	{
+	case EMediaTrackType::Audio:
+		if (AudioTracks.IsValidIndex(TrackIndex))
+		{
+			Track = &AudioTracks[TrackIndex];
+		}
+		break;
+
+	case EMediaTrackType::Caption:
+		if (CaptionTracks.IsValidIndex(TrackIndex))
+		{
+			Track = &CaptionTracks[TrackIndex];
+		}
+		break;
 
-	if (CurrentRefCount == 0)
+	case EMediaTrackType::Video:
+		if (VideoTracks.IsValidIndex(TrackIndex))
+		{
+			Track = &VideoTracks[TrackIndex];
+		}
+		break;
+
+	default:
+		break; // unsupported track type
+	};
+
+	if (Track == nullptr)
+	{
+		return false; // track not found
+	}
+
+	if (Track->SelectedFormat == FormatIndex)
+	{
+		return true; // format already set
+	}
+
+	if (!Track->Formats.IsValidIndex(FormatIndex))
+	{
+		return false; // format not found
+	}
+
+	// get selected format
+	const FFormat& Format = Track->Formats[FormatIndex];
+
+	check(Format.InputType.IsValid());
+	check(Format.OutputType.IsValid());
+	check(Track->Handler.IsValid());
+
+	// set track format
+	HRESULT Result = Track->Handler->SetCurrentMediaType(Format.InputType);
+
+	if (FAILED(Result))
+	{
+		UE_LOG(LogMfMedia, Warning, TEXT("Failed to set selected media type on handler for stream %i (%s)"), Track->StreamIndex, *MfMedia::ResultToString(Result));
+		return false;
+	}
+
+	Result = SourceReader->SetCurrentMediaType(Track->StreamIndex, NULL, Format.OutputType);
+
+	if (FAILED(Result))
 	{
-		delete this;
+		UE_LOG(LogMfMedia, Warning, TEXT("Failed to set selected media type on reader for stream %i (%s)"), Track->StreamIndex, *MfMedia::ResultToString(Result));
+		return false;
 	}
 
-	return CurrentRefCount;
+	Track->SelectedFormat = FormatIndex;
+
+	return true;
 }
 
 
 /* FMfMediaTracks implementation
  *****************************************************************************/
 
-void FMfMediaTracks::AddStreamToTracks(uint32 StreamIndex, IMFPresentationDescriptor* PresentationDescriptor, IMFMediaSource* MediaSource, FString& OutInfo)
+void FMfMediaTracks::AddStreamToTracks(uint32 StreamIndex, FString& OutInfo)
 {
-	// get stream descriptor & media type handler
+	OutInfo += FString::Printf(TEXT("Stream %i\n"), StreamIndex);
+
+	// get stream descriptor
 	TComPtr<IMFStreamDescriptor> StreamDescriptor;
 	{
 		BOOL Selected = FALSE;
-		HRESULT Result = PresentationDescriptor->GetStreamDescriptorByIndex(StreamIndex, &Selected, &StreamDescriptor);
+		const HRESULT Result = PresentationDescriptor->GetStreamDescriptorByIndex(StreamIndex, &Selected, &StreamDescriptor);
 
 		if (FAILED(Result))
 		{
@@ -760,9 +938,10 @@ void FMfMediaTracks::AddStreamToTracks(uint32 StreamIndex, IMFPresentationDescri
 		}
 	}
 
+	// get media type handler
 	TComPtr<IMFMediaTypeHandler> Handler;
 	{
-		HRESULT Result = StreamDescriptor->GetMediaTypeHandler(&Handler);
+		const HRESULT Result = StreamDescriptor->GetMediaTypeHandler(&Handler);
 
 		if (FAILED(Result))
 		{
@@ -776,7 +955,7 @@ void FMfMediaTracks::AddStreamToTracks(uint32 StreamIndex, IMFPresentationDescri
 	// skip unsupported handler types
 	GUID MajorType;
 	{
-		HRESULT Result = Handler->GetMajorType(&MajorType);
+		const HRESULT Result = Handler->GetMajorType(&MajorType);
 
 		if (FAILED(Result))
 		{
@@ -786,7 +965,7 @@ void FMfMediaTracks::AddStreamToTracks(uint32 StreamIndex, IMFPresentationDescri
 			return;
 		}
 
-		UE_LOG(LogMfMedia, Verbose, TEXT("Major type of stream %i is %s"), StreamIndex, *MfMedia::MajorTypeToString(MajorType));
+		UE_LOG(LogMfMedia, Verbose, TEXT("Tracks %llx: Major type of stream %i is %s"), this, StreamIndex, *MfMedia::MajorTypeToString(MajorType));
 		OutInfo += FString::Printf(TEXT("    Type: %s\n"), *MfMedia::MajorTypeToString(MajorType));
 
 		if ((MajorType != MFMediaType_Audio) &&
@@ -800,345 +979,573 @@ void FMfMediaTracks::AddStreamToTracks(uint32 StreamIndex, IMFPresentationDescri
 		}
 	}
 
-	// get media type & make it current
-	TComPtr<IMFMediaType> MediaType;
+	// @todo gmp: handle protected content
+	const bool Protected = ::MFGetAttributeUINT32(StreamDescriptor, MF_SD_PROTECTED, FALSE) != 0;
+	{
+		if (Protected)
+		{
+			OutInfo += FString::Printf(TEXT("    Protected content\n"));
+		}
+	}
+
+	// get number of track formats
+	DWORD NumMediaTypes = 0;
 	{
-		DWORD NumSupportedTypes = 0;
-		HRESULT Result = Handler->GetMediaTypeCount(&NumSupportedTypes);
+		const HRESULT Result = Handler->GetMediaTypeCount(&NumMediaTypes);
 
 		if (FAILED(Result))
 		{
-			UE_LOG(LogMfMedia, Warning, TEXT("Failed to get number of supported media types in stream %i of type %s"), StreamIndex, *MfMedia::MajorTypeToString(MajorType));
-			OutInfo += TEXT("    failed to get supported media types\n");
+			UE_LOG(LogMfMedia, Warning, TEXT("Failed to get number of track formats in stream %i of type %s"), StreamIndex, *MfMedia::MajorTypeToString(MajorType));
+			OutInfo += TEXT("    failed to get track formats\n");
 
 			return;
 		}
+	}
+
+	// create & add track
+	FTrack* Track = nullptr;
+
+	if (MajorType == MFMediaType_Audio)
+	{
+		const int32 TrackIndex = AudioTracks.AddDefaulted();
+		Track = &AudioTracks[TrackIndex];
+	}
+	else if (MajorType == MFMediaType_SAMI)
+	{
+		const int32 TrackIndex = CaptionTracks.AddDefaulted();
+		Track = &CaptionTracks[TrackIndex];
+	}
+	else if (MajorType == MFMediaType_Video)
+	{
+		const int32 TrackIndex = VideoTracks.AddDefaulted();
+		Track = &VideoTracks[TrackIndex];
+	}
 
-		bool MediaFound = false;
-		if (NumSupportedTypes > 0)
+	check(Track != nullptr);
+
+	// get current format
+	TComPtr<IMFMediaType> CurrentMediaType;
+	{
+		if (FAILED(Handler->GetCurrentMediaType(&CurrentMediaType)))
 		{
-			for (DWORD TypeIndex = 0; TypeIndex < NumSupportedTypes; ++TypeIndex)
-			{
-				if (SUCCEEDED(Handler->GetMediaTypeByIndex(TypeIndex, &MediaType)) &&
-					SUCCEEDED(Handler->SetCurrentMediaType(MediaType)))
-				{
-					MediaFound = true;
-					break;
-				}
-			}
+			UE_LOG(LogMfMedia, Warning, TEXT("Failed to get current media type in stream %i of type %s"), StreamIndex, *MfMedia::MajorTypeToString(MajorType));
 		}
+	}
+
+	Track->SelectedFormat = INDEX_NONE;
+
+	// add track formats
+	const bool AllowNonStandardCodecs = false;// GetDefault<UMfMediaSettings>()->AllowNonStandardCodecs;
+
+	for (DWORD TypeIndex = 0; TypeIndex < NumMediaTypes; ++TypeIndex)
+	{
+		OutInfo += FString::Printf(TEXT("    Format %i\n"), TypeIndex);
+
+		// get media type
+		TComPtr<IMFMediaType> MediaType;
 
-		if (MediaType == NULL || !MediaFound)
+		if (FAILED(Handler->GetMediaTypeByIndex(TypeIndex, &MediaType)))
 		{
-			UE_LOG(LogMfMedia, Warning, TEXT("No supported media type in stream %i of type %s"), StreamIndex, *MfMedia::MajorTypeToString(MajorType));
-			OutInfo += TEXT("    unsupported media type\n");
+			OutInfo += TEXT("        failed to get media type\n");
 
-			return;
+			continue;
 		}
-	}
 
-	// get sub-type
-	GUID SubType;
-	{
+		// get sub-type
+		GUID SubType;
+
 		if (MajorType == MFMediaType_SAMI)
 		{
 			FMemory::Memzero(SubType);
 		}
 		else
 		{
-			HRESULT Result = MediaType->GetGUID(MF_MT_SUBTYPE, &SubType);
+			const HRESULT Result = MediaType->GetGUID(MF_MT_SUBTYPE, &SubType);
 
 			if (FAILED(Result))
 			{
-				UE_LOG(LogMfMedia, Warning, TEXT("Failed to get sub-type of stream %i (%s)"), StreamIndex, *MfMedia::ResultToString(Result));
-				OutInfo += TEXT("    failed to get sub-type\n");
+				UE_LOG(LogMfMedia, Warning, TEXT("Failed to get sub-type of format %i in stream %i (%s)"), TypeIndex, StreamIndex, *MfMedia::ResultToString(Result));
+				OutInfo += TEXT("        failed to get sub-type\n");
 
-				return;
+				continue;;
 			}
-
-			UE_LOG(LogMfMedia, Verbose, TEXT("Sub-type of stream %i is %s"), StreamIndex, *MfMedia::SubTypeToString(SubType));
-
-			OutInfo += FString::Printf(TEXT("    Codec: %s\n"), *MfMedia::SubTypeToString(SubType));
-			OutInfo += FString::Printf(TEXT("    Protected: %s\n"), (::MFGetAttributeUINT32(StreamDescriptor, MF_SD_PROTECTED, FALSE) != 0) ? *GYes.ToString() : *GNo.ToString());
 		}
-	}
 
-	// configure desired output type
-	TComPtr<IMFMediaType> OutputType;
-	{
-		HRESULT Result = ::MFCreateMediaType(&OutputType);
+		const FString TypeName = MfMedia::SubTypeToString(SubType);
+		OutInfo += FString::Printf(TEXT("        Codec: %s\n"), *TypeName);
 
-		if (FAILED(Result))
+		// create output type
+		TComPtr<IMFMediaType> OutputType = MfMedia::CreateOutputType(MajorType, SubType, AllowNonStandardCodecs);
+
+		if (!OutputType.IsValid())
 		{
-			UE_LOG(LogMfMedia, Warning, TEXT("Failed to create output type for stream %i (%s)"), StreamIndex, *MfMedia::ResultToString(Result));
-			OutInfo += TEXT("    failed to create output type\n");
+			OutInfo += TEXT("        failed to create output type\n");
 
-			return;
+			continue;
 		}
 
-		Result = OutputType->SetUINT32(MF_MT_ALL_SAMPLES_INDEPENDENT, TRUE);
+		// add format details
+		int32 FormatIndex = INDEX_NONE;
 
-		if (FAILED(Result))
+		if (MajorType == MFMediaType_Audio)
 		{
-			UE_LOG(LogMfMedia, Warning, TEXT("Failed to initialize output type for stream %i (%s)"), StreamIndex, *MfMedia::ResultToString(Result));
-			OutInfo += TEXT("    failed to initialize output type\n");
-
-			return;
+			const uint32 BitsPerSample = ::MFGetAttributeUINT32(MediaType, MF_MT_AUDIO_BITS_PER_SAMPLE, 16);
+			const uint32 NumChannels = ::MFGetAttributeUINT32(MediaType, MF_MT_AUDIO_NUM_CHANNELS, 0);
+			const uint32 SampleRate = ::MFGetAttributeUINT32(MediaType, MF_MT_AUDIO_SAMPLES_PER_SECOND, 0);
+
+			FormatIndex = Track->Formats.Add({
+				MediaType,
+				OutputType,
+				TypeName,
+				{
+					BitsPerSample,
+					NumChannels,
+					SampleRate
+				},
+				{ 0 }
+			});
+
+			OutInfo += FString::Printf(TEXT("        Channels: %i\n"), NumChannels);
+			OutInfo += FString::Printf(TEXT("        Sample Rate: %i Hz\n"), SampleRate);
+			OutInfo += FString::Printf(TEXT("        Bits Per Sample: %i\n"), BitsPerSample);
 		}
-
-		if (MajorType == MFMediaType_Audio)
+		else if (MajorType == MFMediaType_SAMI)
+		{
+			FormatIndex = Track->Formats.Add({
+				MediaType,
+				OutputType,
+				TypeName,
+				{ 0 },
+				{ 0 }
+			});
+		}
+		else if (MajorType == MFMediaType_Video)
 		{
-			// filter unsupported audio formats (XboxOne only supports AAC)
-			if (SubType != MFAudioFormat_AAC)
+			GUID OutputSubType;
 			{
-				UE_LOG(LogMfMedia, Warning, TEXT("Unsupported audio type '%s' (%s) for stream %i"), *MfMediaTracks::FourccToString(SubType.Data1), *MfMediaTracks::GuidToString(SubType), StreamIndex);
-				OutInfo += TEXT("    unsupported SubType\n");
+				const HRESULT Result = OutputType->GetGUID(MF_MT_SUBTYPE, &OutputSubType);
 
-				return;
-			}
-
-			// configure audio output
-			if (FAILED(OutputType->SetGUID(MF_MT_MAJOR_TYPE, MFMediaType_Audio)) ||
-				FAILED(OutputType->SetGUID(MF_MT_SUBTYPE, MFAudioFormat_PCM)) ||
-				FAILED(OutputType->SetUINT32(MF_MT_AUDIO_BITS_PER_SAMPLE, 16u)))
-			{
-				UE_LOG(LogMfMedia, Warning, TEXT("Failed to initialize audio output type for stream %i"), StreamIndex);
-				OutInfo += TEXT("    failed to initialize output type\n");
+				if (FAILED(Result))
+				{
+					UE_LOG(LogMfMedia, Error, TEXT("Failed to get video output sub-type for stream %i (%s)"), StreamIndex, *MfMedia::ResultToString(Result));
+					OutInfo += FString::Printf(TEXT("        failed to get sub-type"));
 
-				return;
+					continue;
+				}
 			}
 
-			if (::MFGetAttributeUINT32(MediaType, MF_MT_AUDIO_SAMPLES_PER_SECOND, 0) != 44100)
+			const uint32 BitRate = ::MFGetAttributeUINT32(MediaType, MF_MT_AVG_BITRATE, 0);
+
+			FIntPoint OutputDim;
 			{
-				UE_LOG(LogMfMedia, Warning, TEXT("Possible loss of audio quality in stream %i due to sample rate != 44100 Hz"), StreamIndex);
+				if (SUCCEEDED(::MFGetAttributeSize(MediaType, MF_MT_FRAME_SIZE, (UINT32*)&OutputDim.X, (UINT32*)&OutputDim.Y)))
+				{
+					OutInfo += FString::Printf(TEXT("        Dimensions: %i x %i\n"), OutputDim.X, OutputDim.Y);
+				}
+				else
+				{
+					OutputDim = FIntPoint::ZeroValue;
+					OutInfo += FString::Printf(TEXT("        Dimensions: n/a\n"));
+				}
 			}
-		}
-		else if (MajorType == MFMediaType_SAMI)
-		{
-			// configure caption output
-			Result = OutputType->SetGUID(MF_MT_MAJOR_TYPE, MFMediaType_SAMI);
 
-			if (FAILED(Result))
+			float FrameRate;
 			{
-				UE_LOG(LogMfMedia, Warning, TEXT("Failed to initialize caption output type for stream %i (%s)"), StreamIndex, *MfMedia::ResultToString(Result));
-				OutInfo += TEXT("    failed to set output type\n");
+				UINT32 Numerator = 0;
+				UINT32 Denominator = 1;
 
-				return;
+				if (SUCCEEDED(::MFGetAttributeRatio(MediaType, MF_MT_FRAME_RATE, &Numerator, &Denominator)))
+				{
+					FrameRate = static_cast<float>(Numerator) / Denominator;
+					OutInfo += FString::Printf(TEXT("        Frame Rate: %g fps\n"), FrameRate);
+				}
+				else
+				{
+					FrameRate = 0.0f;
+					OutInfo += FString::Printf(TEXT("        Frame Rate: n/a\n"));
+				}
 			}
-		}
-		else if (MajorType == MFMediaType_Video)
-		{		
-#if PLATFORM_XBOXONE
-			// filter unsupported video types (XboxOne only supports H.264)
-			if ((SubType != MFVideoFormat_H264) && (SubType != MFVideoFormat_H264_ES))
+
+			TRange<float> FrameRates;
 			{
-				UE_LOG(LogMfMedia, Warning, TEXT("Unsupported video type '%s' (%s) for stream %i"), *MfMediaTracks::FourccToString(SubType.Data1), *MfMediaTracks::GuidToString(SubType), StreamIndex);
-				OutInfo += TEXT("    unsupported SubType\n");
+				UINT32 Numerator = 0;
+				UINT32 Denominator = 1;
+				float Min = -1.0f;
+				float Max = -1.0f;
 
-				return;
-			}
-#endif //PLATFORM_XBOXONE
+				if (SUCCEEDED(::MFGetAttributeRatio(MediaType, MF_MT_FRAME_RATE_RANGE_MIN, &Numerator, &Denominator)))
+				{
+					Min = static_cast<float>(Numerator) / Denominator;
+				}
 
-			// configure video output
-			Result = OutputType->SetGUID(MF_MT_MAJOR_TYPE, MFMediaType_Video);
+				if (SUCCEEDED(::MFGetAttributeRatio(MediaType, MF_MT_FRAME_RATE_RANGE_MAX, &Numerator, &Denominator)))
+				{
+					Max = static_cast<float>(Numerator) / Denominator;
+				}
 
-			if (FAILED(Result))
-			{
-				UE_LOG(LogMfMedia, Warning, TEXT("Failed to set video output type for stream %i (%s)"), StreamIndex, *MfMedia::ResultToString(Result));
-				OutInfo += TEXT("    failed to set output type\n");
+				if ((Min >= 0.0f) && (Max >= 0.0f))
+				{
+					FrameRates = TRange<float>::Inclusive(Min, Max);
+				}
+				else
+				{
+					FrameRates = TRange<float>(FrameRate);
+				}
 
-				return;
+				OutInfo += FString::Printf(TEXT("        Frame Rate Range: %g - %g fps\n"), FrameRates.GetLowerBoundValue(), FrameRates.GetUpperBoundValue());
 			}
 
-#if PLATFORM_XBOXONE
-			Result = OutputType->SetGUID(MF_MT_SUBTYPE, MFVideoFormat_NV12); // XboxOne only supports NV12
-#else
-			Result = OutputType->SetGUID(MF_MT_SUBTYPE, (SubType == MFVideoFormat_RGB32) ? MFVideoFormat_RGB32 : MFVideoFormat_NV12);
+			FIntPoint BufferDim;
+			uint32 BufferStride;
+			EMediaTextureSampleFormat SampleFormat;
+			{
+#if PLATFORM_WINDOWS
+				if (OutputSubType == MFVideoFormat_NV12)
 #endif
+				{
+					BufferDim = FIntPoint(Align(OutputDim.X, 16), Align(OutputDim.Y, 16) * 3 / 2);
+					BufferStride = BufferDim.X;
+					SampleFormat = EMediaTextureSampleFormat::CharNV12;
+				}
+#if PLATFORM_WINDOWS
+				else
+				{
+					long SampleStride = ::MFGetAttributeUINT32(OutputType, MF_MT_DEFAULT_STRIDE, 0);
 
-			if (FAILED(Result))
-			{
-				UE_LOG(LogMfMedia, Warning, TEXT("Failed to set video output sub-type for stream %i (%s)"), StreamIndex, *MfMedia::ResultToString(Result));
-				OutInfo += TEXT("    failed to set output sub-type\n");
+					if (OutputSubType == MFVideoFormat_RGB32)
+					{
+						SampleFormat = EMediaTextureSampleFormat::CharBMP;
+
+						if (SampleStride == 0)
+						{
+							::MFGetStrideForBitmapInfoHeader(OutputSubType.Data1, OutputDim.X, &SampleStride);
+						}
+
+						if (SampleStride == 0)
+						{
+							SampleStride = OutputDim.X * 4;
+						}
+					}
+					else
+					{
+						SampleFormat = EMediaTextureSampleFormat::CharYUY2;
+
+						if (SampleStride == 0)
+						{
+							int32 AlignedOutputX = OutputDim.X;
+
+							if ((SubType == MFVideoFormat_H264) || (SubType == MFVideoFormat_H264_ES))
+							{
+								AlignedOutputX = Align(AlignedOutputX, 16);
+							}
+
+							SampleStride = AlignedOutputX * 2;
+						}
+					}
+
+					if (SampleStride < 0)
+					{
+						SampleStride = -SampleStride;
+					}
 
-				return;
+					BufferDim = FIntPoint(SampleStride / 4, OutputDim.Y);
+					BufferStride = SampleStride;
+				}
+#endif //PLATFORM_WINDOWS
 			}
-		}
-	}
 
-	SourceReader->SetCurrentMediaType(StreamIndex, NULL, OutputType);
+			GUID FormatType = GUID_NULL;
 
-	// create & add track
-	FStreamInfo* StreamInfo = nullptr;
+#if PLATFORM_WINDOWS
+			// prevent duplicates for legacy DirectShow media types
+			// see: https://msdn.microsoft.com/en-us/library/windows/desktop/ff485858(v=vs.85).aspx
 
-	if (MajorType == MFMediaType_Audio)
-	{
-		const int32 AudioTrackIndex = AudioTracks.AddDefaulted();
-		FAudioTrack& AudioTrack = AudioTracks[AudioTrackIndex];
-		{
-			AudioTrack.BitsPerSample = ::MFGetAttributeUINT32(MediaType, MF_MT_AUDIO_BITS_PER_SAMPLE, 16);
-			AudioTrack.NumChannels = ::MFGetAttributeUINT32(MediaType, MF_MT_AUDIO_NUM_CHANNELS, 0);
-			AudioTrack.SampleRate = ::MFGetAttributeUINT32(MediaType, MF_MT_AUDIO_SAMPLES_PER_SECOND, 0);
-		}
+			if (SUCCEEDED(MediaType->GetGUID(MF_MT_AM_FORMAT_TYPE, &FormatType)))
+			{
+				if (FormatType == FORMAT_VideoInfo)
+				{
+					for (int32 Index = Track->Formats.Num() - 1; Index >= 0; --Index)
+					{
+						const FFormat& Format = Track->Formats[Index];
 
-		OutInfo += FString::Printf(TEXT("    Channels: %i\n"), AudioTrack.NumChannels);
-		OutInfo += FString::Printf(TEXT("    Sample Rate: %i Hz\n"), AudioTrack.SampleRate);
-		OutInfo += FString::Printf(TEXT("    Bits Per Sample: %i\n"), AudioTrack.BitsPerSample);
+						if ((Format.Video.FormatType == FORMAT_VideoInfo2) &&
+							(Format.Video.FrameRates == FrameRates) &&
+							(Format.Video.OutputDim == OutputDim) &&
+							(Format.TypeName == TypeName))
+						{
+							FormatIndex = Index; // keep newer format
 
-		StreamInfo = &AudioTrack;
-	}
-	else if (MajorType == MFMediaType_SAMI)
-	{
-		const int32 CaptionTrackIndex = CaptionTracks.AddDefaulted();
-		FCaptionTrack& CaptionTrack = CaptionTracks[CaptionTrackIndex];
+							break;
+						}
+					}
+				}
+				else if (FormatType == FORMAT_VideoInfo2)
+				{
+					for (int32 Index = Track->Formats.Num() - 1; Index >= 0; --Index)
+					{
+						FFormat& Format = Track->Formats[Index];
 
-		StreamInfo = &CaptionTrack;
-	}
-	else if ((MajorType == MFMediaType_Video) || (MajorType == MFMediaType_Image))
-	{
-		GUID OutputSubType;
-		{
-			HRESULT Result = OutputType->GetGUID(MF_MT_SUBTYPE, &OutputSubType);
+						if ((Format.Video.FormatType == FORMAT_VideoInfo) &&
+							(Format.Video.FrameRates == FrameRates) &&
+							(Format.Video.OutputDim == OutputDim) &&
+							(Format.TypeName == TypeName))
+						{
+							Format.InputType = MediaType; // replace legacy format
+							FormatIndex = Index;
 
-			if (FAILED(Result))
+							break;
+						}
+					}
+				}
+			}
+#endif //PLATFORM_WINDOWS
+
+			if (FormatIndex == INDEX_NONE)
 			{
-				UE_LOG(LogMfMedia, Error, TEXT("Failed to get video output sub-type for stream %i (%s)"), StreamIndex, *MfMedia::ResultToString(Result));
-				return;
+				FormatIndex = Track->Formats.Add({
+					MediaType,
+					OutputType,
+					TypeName,
+					{ 0 },
+					{
+						BitRate,
+						BufferDim,
+						BufferStride,
+						FormatType,
+						FrameRate,
+						FrameRates,
+						OutputDim,
+						SampleFormat
+					}
+				});
 			}
 		}
+		else
+		{
+			check(false); // should never get here
+		}
 
-		const int32 VideoTrackIndex = VideoTracks.AddDefaulted();
-		FVideoTrack& VideoTrack = VideoTracks[VideoTrackIndex];
+		if (MediaType == CurrentMediaType)
 		{
-			// bit rate
-			VideoTrack.BitRate = ::MFGetAttributeUINT32(MediaType, MF_MT_AVG_BITRATE, 0);
+			const HRESULT Result = SourceReader->SetCurrentMediaType(StreamIndex, NULL, OutputType);
 
-			// dimensions
-			if (SUCCEEDED(::MFGetAttributeSize(MediaType, MF_MT_FRAME_SIZE, (UINT32*)&VideoTrack.OutputDim.X, (UINT32*)&VideoTrack.OutputDim.Y)))
+			if (SUCCEEDED(Result))
 			{
-				OutInfo += FString::Printf(TEXT("    Dimensions: %i x %i\n"), VideoTrack.OutputDim.X, VideoTrack.OutputDim.Y);
+				Track->SelectedFormat = FormatIndex;
 			}
 			else
 			{
-				VideoTrack.OutputDim = FIntPoint::ZeroValue;
-				OutInfo += FString::Printf(TEXT("    Dimensions: n/a\n"));
+				UE_LOG(LogMfMedia, Warning, TEXT("Failed to set current media type on reader for stream %i (%s)"), Track->StreamIndex, *MfMedia::ResultToString(Result));
 			}
+		}
+	}
 
-			// frame rate
-			UINT32 Numerator = 0;
-			UINT32 Denominator = 1;
+	// get stream info
+	PWSTR OutString = NULL;
+	UINT32 OutLength = 0;
 
-			if (SUCCEEDED(::MFGetAttributeRatio(MediaType, MF_MT_FRAME_RATE, &Numerator, &Denominator)))
-			{
-				VideoTrack.FrameRate = static_cast<float>(Numerator) / Denominator;
-				OutInfo += FString::Printf(TEXT("    Frame Rate: %g fps\n"), VideoTrack.FrameRate);
-			}
-			else
-			{
-				VideoTrack.FrameRate = 0.0f;
-				OutInfo += FString::Printf(TEXT("    Frame Rate: n/a\n"));
-			}
+	if (SUCCEEDED(StreamDescriptor->GetAllocatedString(MF_SD_LANGUAGE, &OutString, &OutLength)))
+	{
+		Track->Language = OutString;
+		::CoTaskMemFree(OutString);
+	}
 
-#if PLATFORM_XBOXONE
-			VideoTrack.BufferDim = FIntPoint(Align(VideoTrack.OutputDim.X, 16), Align(VideoTrack.OutputDim.Y, 16) * 3 / 2);
-			VideoTrack.SinkFormat = EMediaTextureSinkFormat::CharNV12;
-#else
-			if (OutputSubType == MFVideoFormat_NV12)
-			{
-				VideoTrack.BufferDim = FIntPoint(Align(VideoTrack.OutputDim.X, 16), Align(VideoTrack.OutputDim.Y, 16) * 3 / 2);
-				VideoTrack.SinkFormat = EMediaTextureSinkFormat::CharNV12;
-			}
-			else
-			{
-				long SampleStride = ::MFGetAttributeUINT32(MediaType, MF_MT_DEFAULT_STRIDE, 0);
+	if (SUCCEEDED(StreamDescriptor->GetAllocatedString(MF_SD_STREAM_NAME, &OutString, &OutLength)))
+	{
+		Track->Name = OutString;
+		::CoTaskMemFree(OutString);
+	}
 
-				if (SampleStride == 0)
-				{
-					::MFGetStrideForBitmapInfoHeader(SubType.Data1, VideoTrack.OutputDim.X, &SampleStride);
-				}
+	Track->DisplayName = (Track->Name.IsEmpty())
+		? FText::Format(LOCTEXT("UnnamedStreamFormat", "Unnamed Track (Stream {0})"), FText::AsNumber((uint32)StreamIndex))
+		: FText::FromString(Track->Name);
 
-				if (SampleStride == 0)
-				{
-					SampleStride = VideoTrack.OutputDim.X * 4;
-				}
-				else if (SampleStride < 0)
-				{
-					SampleStride = -SampleStride;
-				}
+	Track->Descriptor = StreamDescriptor;
+	Track->Handler = Handler;
+	Track->Protected = Protected;
+	Track->StreamIndex = StreamIndex;
 
-				VideoTrack.BufferDim = FIntPoint(SampleStride / 4, VideoTrack.OutputDim.Y);
-				VideoTrack.SinkFormat = EMediaTextureSinkFormat::CharBMP;
-			}
-#endif //PLATFORM_XBOXONE	
-		}
+	if (Track->Formats.Num() > 0)
+	{
+		Track->SelectedFormat = 0;
+	}
+	else
+	{
+		Track->SelectedFormat = INDEX_NONE;
 
-		StreamInfo = &VideoTrack;
+		UE_LOG(LogMfMedia, Warning, TEXT("No supported media type in stream %i of type %s"), StreamIndex, *MfMedia::MajorTypeToString(MajorType));
+		OutInfo += TEXT("    unsupported media type\n");
 	}
+}
 
-	// get stream info
-	if (StreamInfo != nullptr)
+
+const FMfMediaTracks::FFormat* FMfMediaTracks::GetAudioFormat(int32 TrackIndex, int32 FormatIndex) const
+{
+	if (AudioTracks.IsValidIndex(TrackIndex))
+	{
+		const FTrack& Track = AudioTracks[TrackIndex];
+
+		if (Track.Formats.IsValidIndex(FormatIndex))
+		{
+			return &Track.Formats[FormatIndex];
+		}
+	}
+
+	return nullptr;
+}
+
+
+const FMfMediaTracks::FTrack* FMfMediaTracks::GetTrack(EMediaTrackType TrackType, int32 TrackIndex) const
+{
+	switch (TrackType)
 	{
-		PWSTR OutString = NULL;
-		UINT32 OutLength;
+	case EMediaTrackType::Audio:
+		if (AudioTracks.IsValidIndex(TrackIndex))
+		{
+			return &AudioTracks[TrackIndex];
+		}
 
-		if (SUCCEEDED(StreamDescriptor->GetAllocatedString(MF_SD_LANGUAGE, &OutString, &OutLength)))
+	case EMediaTrackType::Caption:
+		if (CaptionTracks.IsValidIndex(TrackIndex))
 		{
-			StreamInfo->Language = OutString;
-			::CoTaskMemFree(OutString);
+			return &CaptionTracks[TrackIndex];
 		}
 
-		if (SUCCEEDED(StreamDescriptor->GetAllocatedString(MF_SD_STREAM_NAME, &OutString, &OutLength)))
+	case EMediaTrackType::Video:
+		if (VideoTracks.IsValidIndex(TrackIndex))
 		{
-			StreamInfo->Name = OutString;
-			::CoTaskMemFree(OutString);
+			return &VideoTracks[TrackIndex];
 		}
 
-		StreamInfo->DisplayName = (StreamInfo->Name.IsEmpty())
-			? FText::Format(LOCTEXT("UnnamedStreamFormat", "Unnamed Stream {0}"), FText::AsNumber((uint32)StreamIndex))
-			: FText::FromString(StreamInfo->Name);
+	default:
+		break; // unsupported track type
+	}
+
+	return nullptr;
+}
+
+
+const FMfMediaTracks::FFormat* FMfMediaTracks::GetVideoFormat(int32 TrackIndex, int32 FormatIndex) const
+{
+	if (VideoTracks.IsValidIndex(TrackIndex))
+	{
+		const FTrack& Track = VideoTracks[TrackIndex];
 
-		StreamInfo->SecondsUntilNextSample = 0.f;
-		StreamInfo->StreamIndex = StreamIndex;
+		if (Track.Formats.IsValidIndex(FormatIndex))
+		{
+			return &Track.Formats[FormatIndex];
+		}
 	}
+
+	return nullptr;
 }
 
 
-void FMfMediaTracks::InitializeAudioSink()
+bool FMfMediaTracks::RequestSample(DWORD StreamIndex)
 {
-	if ((AudioSink == nullptr) || (SelectedAudioTrack == INDEX_NONE))
+#if MFMEDIATRACKS_USE_ASYNCREADER
+	const HRESULT Result = SourceReader->ReadSample(StreamIndex, 0, nullptr, nullptr, nullptr, nullptr);
+
+#else
+	TComPtr<IMFSample> Sample;
+	DWORD StreamFlags;
+	LONGLONG Timestamp;
+
+	const HRESULT Result = SourceReader->ReadSample(StreamIndex, 0, nullptr, &StreamFlags, &Timestamp, &Sample);
+#endif
+
+	if (FAILED(Result))
 	{
-		return;
+		UE_LOG(LogMfMedia, VeryVerbose, TEXT("Tracks %llx: Failed to request sample for stream %i: %s"), this, StreamIndex, *MfMedia::ResultToString(Result));
+		return false;
+	}
+
+#if MFMEDIATRACKS_USE_ASYNCREADER
+	return true;
+
+#else
+	if (SUCCEEDED(Result))
+	{
+		ProcessSample(Sample, Result, StreamFlags, StreamIndex, FTimespan(Timestamp));
+	}
+
+	return false;
+#endif
+}
+
+
+void FMfMediaTracks::UpdateAudio()
+{
+	if (AudioSamplePending)
+	{
+		return; // sample request pending
+	}
+
+	if (AudioDone || AudioSampleRange.IsEmpty() || (SelectedAudioTrack == INDEX_NONE))
+	{
+		return; // nothing to play
 	}
 
-	const FAudioTrack& AudioTrack = AudioTracks[SelectedAudioTrack];
-	AudioSink->InitializeAudioSink(AudioTrack.NumChannels, AudioTrack.SampleRate);
-	if (Enabled)
+	if ((LastAudioSampleTime != FTimespan::MinValue()) && !AudioSampleRange.Contains(LastAudioSampleTime))
 	{
-		AudioSink->ResumeAudioSink();
+		return; // no new sample needed
 	}
+
+	#if MFMEDIATRACKS_TRACE_SAMPLES
+		UE_LOG(LogMfMedia, VeryVerbose, TEXT("Tracks %llx: Requesting audio sample"), this);
+	#endif
+
+	AudioSamplePending = RequestSample(AudioTracks[SelectedAudioTrack].StreamIndex);
 }
 
 
-void FMfMediaTracks::InitializeOverlaySink()
+void FMfMediaTracks::UpdateCaptions()
 {
-	if ((OverlaySink == nullptr) || (SelectedCaptionTrack == INDEX_NONE))
+	if (CaptionSamplePending)
 	{
-		return;
+		return; // sample request pending
+	}
+
+	if (CaptionDone || CaptionSampleRange.IsEmpty() || (SelectedCaptionTrack == INDEX_NONE))
+	{
+		return; // nothing to play
 	}
 
-	OverlaySink->InitializeOverlaySink();
+	if ((LastCaptionSampleTime != FTimespan::MinValue()) && !CaptionSampleRange.Contains(LastCaptionSampleTime))
+	{
+		return; // no new sample needed
+	}
+
+	#if MFMEDIATRACKS_TRACE_SAMPLES
+		UE_LOG(LogMfMedia, VeryVerbose, TEXT("Tracks %llx: Requesting caption sample"), this);
+	#endif
+
+	CaptionSamplePending = RequestSample(CaptionTracks[SelectedCaptionTrack].StreamIndex);
 }
 
 
-void FMfMediaTracks::InitializeVideoSink()
+void FMfMediaTracks::UpdateVideo()
 {
-	if ((VideoSink == nullptr) || (SelectedVideoTrack == INDEX_NONE))
+	if (VideoSamplePending)
 	{
-		return;
+		return; // sample request pending
+	}
+
+	if (VideoDone || VideoSampleRange.IsEmpty() || (SelectedVideoTrack == INDEX_NONE))
+	{
+		return; // nothing to play
 	}
 
-	const auto& VideoTrack = VideoTracks[SelectedVideoTrack];
-	VideoSink->InitializeTextureSink(VideoTrack.OutputDim, VideoTrack.BufferDim, VideoTrack.SinkFormat, EMediaTextureSinkMode::Unbuffered);
+	if ((LastVideoSampleTime != FTimespan::MinValue()) && !VideoSampleRange.Contains(LastVideoSampleTime))
+	{
+		return; // no new sample needed
+	}
+
+	#if MFMEDIATRACKS_TRACE_SAMPLES
+		UE_LOG(LogMfMedia, VeryVerbose, TEXT("Tracks %llx: Requesting video sample"), this);
+	#endif
+
+	VideoSamplePending = RequestSample(VideoTracks[SelectedVideoTrack].StreamIndex);
 }
 
 
