Copying //UE4/Dev-Sequencer to //UE4/Main (Source: //UE4/Dev-Sequencer @ 3617642)

diff --git a/Engine/Plugins/Media/MfMedia/Source/MfMedia/Private/Player/MfMediaTracks.h b/Engine/Plugins/Media/MfMedia/Source/MfMedia/Private/Player/MfMediaTracks.h
--- a/Engine/Plugins/Media/MfMedia/Source/MfMedia/Private/Player/MfMediaTracks.h
+++ b/Engine/Plugins/Media/MfMedia/Source/MfMedia/Private/Player/MfMediaTracks.h
@@ -6,8 +6,15 @@
 
 #if MFMEDIA_SUPPORTED_PLATFORM
 
-#include "IMediaOutput.h"
+#include "CoreTypes.h"
+#include "Containers/Array.h"
+#include "Containers/UnrealString.h"
+#include "HAL/CriticalSection.h"
 #include "IMediaTracks.h"
+#include "Internationalization/Text.h"
+#include "Math/IntPoint.h"
+#include "Math/Range.h"
+#include "Windows/COMPointer.h"
 
 #if PLATFORM_WINDOWS
 	#include "WindowsHWrapper.h"
@@ -16,118 +23,202 @@
 	#include "XboxOneAllowPlatformTypes.h"
 #endif
 
-class IMediaOptions;
+class FMediaSamples;
+class FMfMediaAudioSamplePool;
+class FMfMediaTextureSamplePool; 
+
+enum class EMediaTextureSampleFormat;
+
+struct FMfMediaSourceReaderSample;
 
 
 /**
- * Implements the track collection for Microsoft Media Framework based media players.
+ * Track collection for Media Foundation based media players.
  */
 class FMfMediaTracks
-	: public IMediaOutput
-	, public IMediaTracks
+	: public IMediaTracks
 {
-	struct FStreamInfo
+	struct FFormat
+	{
+		TComPtr<IMFMediaType> InputType;
+		TComPtr<IMFMediaType> OutputType;
+		FString TypeName;
+
+		struct
+		{
+			uint32 BitsPerSample;
+			uint32 NumChannels;
+			uint32 SampleRate;
+		}
+		Audio;
+
+		struct 
+		{
+			uint32 BitRate;
+			FIntPoint BufferDim;
+			uint32 BufferStride;
+			GUID FormatType;
+			float FrameRate;
+			TRange<float> FrameRates;
+			FIntPoint OutputDim;
+			EMediaTextureSampleFormat SampleFormat;
+		}
+		Video;
+	};
+
+	struct FTrack
 	{
+		TComPtr<IMFStreamDescriptor> Descriptor;
 		FText DisplayName;
+		TArray<FFormat> Formats;
+		TComPtr<IMFMediaTypeHandler> Handler;
 		FString Language;
 		FString Name;
-		float SecondsUntilNextSample;
+		bool Protected;
+		int32 SelectedFormat;
 		int32 StreamIndex;
 	};
 
-	struct FAudioTrack : FStreamInfo
-	{
-		uint32 BitsPerSample;
-		uint32 NumChannels;
-		uint32 SampleRate;
-	};
-
-	struct FCaptionTrack : FStreamInfo
-	{
-	};
-
-	struct FVideoTrack : FStreamInfo
-	{
-		uint32 BitRate;
-		FIntPoint BufferDim;
-		float FrameRate;
-		FIntPoint OutputDim;
-		EMediaTextureSinkFormat SinkFormat;
-	};
-
 public:
 
 	/** Default constructor. */
 	FMfMediaTracks();
 
-	/** Destructor. */
-	~FMfMediaTracks();
+	/** Virtual destructor. */
+	virtual ~FMfMediaTracks();
 
 public:
 
 	/**
-	 * Initialize this track collection.
+	 * Append track statistics information to the given string.
 	 *
-	 * @param InPresentationDescriptor The presentation descriptor object.
-	 * @param InMediaSourceObject The media source object.
-	 * @param InSourceReader The source reader to use.
-	 * @param OutInfo String to append debug information to.
-	 * @see Reinitialize
+	 * @param OutStats The string to append the statistics to.
 	 */
-	void Initialize(IMFPresentationDescriptor* InPresentationDescriptor, IMFMediaSource* InMediaSource, IMFSourceReader* InSourceReader, FString& OutInfo);
+	void AppendStats(FString &OutStats) const;
 
 	/**
-	 * Reinitialize with the previous presentation descriptor and media source.
+	 * Clear the streams flags.
 	 *
-	 * @see Initialize
+	 * @see GetFlags
 	 */
-	void Reinitialize();
+	void ClearFlags();
 
-	/** Reset this object. */
-	void Reset();
+	/**
+	 * Get the total duration of the current media source.
+	 *
+	 * @return Media duration.
+	 * @see GetInfo, Initialize
+	 */
+	FTimespan GetDuration() const;
+
+	/**
+	 * Get the current flags.
+	 *
+	 * @param OutMediaSourceChanged Will indicate whether the media source changed.
+	 * @param OutSelectionChanged Will indicate whether the track selection changed.
+	 * @see ClearFlags
+	 */
+	void GetFlags(bool& OutMediaSourceChanged, bool& OutSelectionChanged) const;
 
 	/**
-	 * Set whether tracks are enabled.
+	 * Get the information string for the currently loaded media source.
 	 *
-	 * @param InEnabled true if tracks are enabled, false otherwise.
+	 * @return Info string.
+	 * @see GetDuration, GetSamples
 	 */
-	void SetEnabled(bool InEnabled);
+	const FString& GetInfo() const
+	{
+		return Info;
+	}
 
-	/** Called each frame. */
-	void Tick(float DeltaTime);
+	/**
+	 * Get the current media source object.
+	 *
+	 * @return The media source, or nullptr if not initialized.
+	 * @see GetSourceReader
+	 */
+	IMFMediaSource* GetMediaSource();
 
-public:
+	/**
+	 * Get the current source reader object.
+	 *
+	 * @return The source reader, or nullptr if not initialized.
+	 * @see GetMediaSource
+	 */
+	IMFSourceReader* GetSourceReader();
 
-	//~ IMediaOutput interface
+	/**
+	 * Initialize the track collection.
+	 *
+	 * @param InMediaSource The media source object.
+	 * @param InSourceReaderCallback The player's source reader callback.
+	 * @param InSamples The sample collection that receives output samples.
+	 * @see IsInitialized, Shutdown
+	 */
+	void Initialize(IMFMediaSource* InMediaSource, IMFSourceReaderCallback* InSourceReaderCallback, const TSharedRef<FMediaSamples, ESPMode::ThreadSafe>& InSamples);
 
-	virtual void SetAudioSink(IMediaAudioSink* Sink) override;
-	virtual void SetMetadataSink(IMediaBinarySink* Sink) override;
-	virtual void SetOverlaySink(IMediaOverlaySink* Sink) override;
-	virtual void SetVideoSink(IMediaTextureSink* Sink) override;
+	/**
+	 * Whether this object has been initialized.
+	 *
+	 * @return true if initialized, false otherwise.
+	 * @see Initialize, Shutdown
+	 */
+	bool IsInitialized() const
+	{
+		return (MediaSource != NULL);
+	}
+
+	/**
+	 * Process a media sample from the source reader callback.
+	 *
+	 * @param ReaderSample The sample to process.
+	 */
+	void ProcessSample(IMFSample* Sample, HRESULT Status, DWORD StreamFlags, DWORD StreamIndex, FTimespan Time);
+
+	/** Restart stream sampling. */
+	void Restart();
+
+	/**
+	 * Shut down the track collection.
+	 *
+	 * @see Initialize, IsInitialized
+	 */
+	void Shutdown();
+
+	/**
+	 * Tick audio sample processing.
+	 *
+	 * @param Rate The current playback rate.
+	 * @param Time The current presentation time.
+	 * @see TickVideo
+	 */
+	void TickAudio(float Rate, FTimespan Time);
+
+	/**
+	 * Tick caption & video sample processing.
+	 *
+	 * @param Rate The current playback rate.
+	 * @param Time The current presentation time.
+	 * @param DeltaTime Time since the last tick.
+	 * @see TickAudio
+	 */
+	void TickInput(float Rate, FTimespan Time, FTimespan DeltaTime);
 
 public:
 
 	//~ IMediaTracks interface
 
-	virtual uint32 GetAudioTrackChannels(int32 TrackIndex) const override;
-	virtual uint32 GetAudioTrackSampleRate(int32 TrackIndex) const override;
+	virtual bool GetAudioTrackFormat(int32 TrackIndex, int32 FormatIndex, FMediaAudioTrackFormat& OutFormat) const override;
 	virtual int32 GetNumTracks(EMediaTrackType TrackType) const override;
+	virtual int32 GetNumTrackFormats(EMediaTrackType TrackType, int32 TrackIndex) const override;
 	virtual int32 GetSelectedTrack(EMediaTrackType TrackType) const override;
 	virtual FText GetTrackDisplayName(EMediaTrackType TrackType, int32 TrackIndex) const override;
+	virtual int32 GetTrackFormat(EMediaTrackType TrackType, int32 TrackIndex) const override;
 	virtual FString GetTrackLanguage(EMediaTrackType TrackType, int32 TrackIndex) const override;
 	virtual FString GetTrackName(EMediaTrackType TrackType, int32 TrackIndex) const override;
-	virtual uint32 GetVideoTrackBitRate(int32 TrackIndex) const override;
-	virtual FIntPoint GetVideoTrackDimensions(int32 TrackIndex) const override;
-	virtual float GetVideoTrackFrameRate(int32 TrackIndex) const override;
+	virtual bool GetVideoTrackFormat(int32 TrackIndex, int32 FormatIndex, FMediaVideoTrackFormat& OutFormat) const override;
 	virtual bool SelectTrack(EMediaTrackType TrackType, int32 TrackIndex) override;
-
-public:
-
-	//~ IUnknown interface
-
-	STDMETHODIMP QueryInterface(REFIID RefID, void** Object);
-	STDMETHODIMP_(ULONG) AddRef();
-	STDMETHODIMP_(ULONG) Release();
+	virtual bool SetTrackFormat(EMediaTrackType TrackType, int32 TrackIndex, int32 FormatIndex) override;
 
 protected:
 
@@ -135,45 +226,123 @@ protected:
 	 * Adds the specified stream to the track collection.
 	 *
 	 * @param StreamIndex The index number of the stream in the presentation descriptor.
-	 * @param PresentationDescriptor The presentation descriptor object.
-	 * @param MediaSource The media source object.
-	 * @param SourceReader The source reader to use.
 	 * @param OutInfo String to append debug information to.
 	 */
-	void AddStreamToTracks(uint32 StreamIndex, IMFPresentationDescriptor* PresentationDescriptor, IMFMediaSource* MediaSource, FString& OutInfo);
+	void AddStreamToTracks(uint32 StreamIndex, FString& OutInfo);
 
-	/** Initialize the current audio sink. */
-	void InitializeAudioSink();
+	/**
+	 * Get the specified audio format.
+	 *
+	 * @param TrackIndex Index of the audio track that contains the format.
+	 * @param FormatIndex Index of the format to return.
+	 * @return Pointer to format, or nullptr if not found.
+	 * @see GetVideoFormat
+	 */
+	const FFormat* GetAudioFormat(int32 TrackIndex, int32 FormatIndex) const;
 
-	/** Initialize the current text overlay sink. */
-	void InitializeOverlaySink();
+	/**
+	 * Get the specified track information.
+	 *
+	 * @param TrackType The type of track.
+	 * @param TrackIndex Index of the track to return.
+	 * @return Pointer to track, or nullptr if not found.
+	 */
+	const FTrack* GetTrack(EMediaTrackType TrackType, int32 TrackIndex) const;
 
-	/** Initialize the current video sink. */
-	void InitializeVideoSink();
+	/**
+	 * Get the specified video format.
+	 *
+	 * @param TrackIndex Index of the video track that contains the format.
+	 * @param FormatIndex Index of the format to return.
+	 * @return Pointer to format, or nullptr if not found.
+	 * @see GetAudioFormat
+	 */
+	const FFormat* GetVideoFormat(int32 TrackIndex, int32 FormatIndex) const;
 
-private:
+	/**
+	 * Request a new media sample for the specified stream.
+	 *
+	 * @param StreamIndex The index of the stream to request a sample for.
+	 * @return true if a sample was requested, false otherwise.
+	 */
+	bool RequestSample(DWORD StreamIndex);
 
-	/** The currently used audio sink. */
-	IMediaAudioSink* AudioSink;
+	/**
+	 * Request more audio samples, if needed.
+	 *
+	 * @see UpdateCaptions, UpdateVideo
+	 */
+	void UpdateAudio();
 
-	/** The currently used text overlay sink. */
-	IMediaOverlaySink* OverlaySink;
+	/**
+	 * Request more caption samples, if needed.
+	 *
+	 * @see UpdateAudio, UpdateVideo
+	 */
+	void UpdateCaptions();
 
-	/** The currently used video sink. */
-	IMediaTextureSink* VideoSink;
+	/**
+	 * Request more video samples, if needed.
+	 *
+	 * @see UpdateAudio, UpdateCaptions
+	 */
+	void UpdateVideo();
 
 private:
 
+	/** Whether the audio track has reached the end. */
+	bool AudioDone;
+
+	/** Whether the next audio sample has been requested. */
+	bool AudioSamplePending;
+
+	/** Audio sample object pool. */
+	FMfMediaAudioSamplePool* AudioSamplePool;
+
+	/** Time range of audio samples to be requested. */
+	TRange<FTimespan> AudioSampleRange;
+
 	/** The available audio tracks. */
-	TArray<FAudioTrack> AudioTracks;
+	TArray<FTrack> AudioTracks;
+
+	/** Whether the caption track has reached the end. */
+	bool CaptionDone;
+
+	/** Whether the next caption sample has been requested. */
+	bool CaptionSamplePending;
+
+	/** Time range of caption samples to be requested. */
+	TRange<FTimespan> CaptionSampleRange;
 
 	/** The available caption tracks. */
-	TArray<FCaptionTrack> CaptionTracks;
+	TArray<FTrack> CaptionTracks;
 
-	/** The available video tracks. */
-	TArray<FVideoTrack> VideoTracks;
+	/** Synchronizes write access to track arrays, selections & sinks. */
+	mutable FCriticalSection CriticalSection;
 
-private:
+	/** Media information string. */
+	FString Info;
+
+	/** Timestamp of the last processed audio sample. */
+	FTimespan LastAudioSampleTime;
+
+	/** Timestamp of the last processed caption sample. */
+	FTimespan LastCaptionSampleTime;
+
+	/** Timestamp of the last processed video sample. */
+	FTimespan LastVideoSampleTime;
+
+	/** The currently opened media. */
+	TComPtr<IMFMediaSource> MediaSource;
+
+	/** Whether the media source has changed. */
+	bool MediaSourceChanged;
+
+	/** The presentation descriptor of the currently opened media. */
+	TComPtr<IMFPresentationDescriptor> PresentationDescriptor;
+
+	/** Media sample collection that receives the output. */
+	TSharedPtr<FMediaSamples, ESPMode::ThreadSafe> Samples;
 
 	/** Index of the selected audio track. */
 	int32 SelectedAudioTrack;
@@ -184,22 +353,8 @@ private:
 	/** Index of the selected video track. */
 	int32 SelectedVideoTrack;
 
-private:
-
-	/** Whether the audio track has reached the end. */
-	bool AudioDone;
-
-	/** Whether the caption track has reached the end. */
-	bool CaptionDone;
-
-	/** Synchronizes write access to track arrays, selections & sinks. */
-	FCriticalSection CriticalSection;
-
-	/** Whether tracks are enabled. */
-	bool Enabled;
-
-	/** Holds a reference counter for this instance. */
-	int32 RefCount;
+	/** Whether the track selection changed. */
+	bool SelectionChanged;
 
 	/** The source reader to use. */
 	TComPtr<IMFSourceReader> SourceReader;
@@ -207,8 +362,17 @@ private:
 	/** Whether the video track has reached the end. */
 	bool VideoDone;
 
-	/** Whether the tracks have started playback. Set on first enabled tick. */
-	bool bIsStarted;
+	/** Whether the next video sample has been requested. */
+	bool VideoSamplePending;
+
+	/** Video sample object pool. */
+	FMfMediaTextureSamplePool* VideoSamplePool;
+
+	/** Time range of video samples to be requested. */
+	TRange<FTimespan> VideoSampleRange;
+
+	/** The available video tracks. */
+	TArray<FTrack> VideoTracks;
 };
 
 
diff --git a/Engine/Plugins/Media/MfMedia/Source/MfMedia/Public/IMfMediaModule.h b/Engine/Plugins/Media/MfMedia/Source/MfMedia/Public/IMfMediaModule.h
--- a/Engine/Plugins/Media/MfMedia/Source/MfMedia/Public/IMfMediaModule.h
+++ b/Engine/Plugins/Media/MfMedia/Source/MfMedia/Public/IMfMediaModule.h
@@ -2,12 +2,13 @@
 
 #pragma once
 
-#include "CoreTypes.h"
 #include "Templates/SharedPointer.h"
 #include "Modules/ModuleInterface.h"
 
+class IMediaEventSink;
 class IMediaPlayer;
 
+
 /**
  * Interface for the MfMedia module.
  */
@@ -17,11 +18,12 @@ class IMfMediaModule
 public:
 
 	/**
-	 * Creates an Xbox One Media Foundation based media player.
+	 * Create an Xbox One Media Foundation based media player.
 	 *
+	 * @param EventSink The object that receives media events from the player.
 	 * @return A new media player, or nullptr if a player couldn't be created.
 	 */
-	virtual TSharedPtr<IMediaPlayer, ESPMode::ThreadSafe> CreatePlayer() = 0;
+	virtual TSharedPtr<IMediaPlayer, ESPMode::ThreadSafe> CreatePlayer(IMediaEventSink& EventSink) = 0;
 
 public:
 
diff --git a/Engine/Plugins/Media/MfMedia/Source/MfMediaEditor/Private/MfFileMediaSourceFactory.cpp b/Engine/Plugins/Media/MfMedia/Source/MfMediaEditor/Private/MfFileMediaSourceFactory.cpp
--- a/Engine/Plugins/Media/MfMedia/Source/MfMediaEditor/Private/MfFileMediaSourceFactory.cpp
+++ b/Engine/Plugins/Media/MfMedia/Source/MfMediaEditor/Private/MfFileMediaSourceFactory.cpp
@@ -1,7 +1,9 @@
 // Copyright 1998-2017 Epic Games, Inc. All Rights Reserved.
 
 #include "MfFileMediaSourceFactory.h"
+
 #include "FileMediaSource.h"
+#include "UObject/UObjectGlobals.h"
 
 
 /* UMfFileMediaSourceFactory structors
diff --git a/Engine/Plugins/Media/MfMedia/Source/MfMediaEditor/Private/MfFileMediaSourceFactory.h b/Engine/Plugins/Media/MfMedia/Source/MfMediaEditor/Private/MfFileMediaSourceFactory.h
--- a/Engine/Plugins/Media/MfMedia/Source/MfMediaEditor/Private/MfFileMediaSourceFactory.h
+++ b/Engine/Plugins/Media/MfMedia/Source/MfMediaEditor/Private/MfFileMediaSourceFactory.h
@@ -2,11 +2,12 @@
 
 #pragma once
 
-#include "CoreMinimal.h"
 #include "UObject/ObjectMacros.h"
 #include "Factories/Factory.h"
+
 #include "MfFileMediaSourceFactory.generated.h"
 
+
 /**
  * Implements a factory for UFileMediaSource objects.
  */
diff --git a/Engine/Plugins/Media/MfMedia/Source/MfMediaFactory/Private/MfMediaFactoryModule.cpp b/Engine/Plugins/Media/MfMedia/Source/MfMediaFactory/Private/MfMediaFactoryModule.cpp
--- a/Engine/Plugins/Media/MfMedia/Source/MfMediaFactory/Private/MfMediaFactoryModule.cpp
+++ b/Engine/Plugins/Media/MfMedia/Source/MfMediaFactory/Private/MfMediaFactoryModule.cpp
@@ -1,14 +1,19 @@
 // Copyright 1998-2017 Epic Games, Inc. All Rights Reserved.
 
-#include "CoreMinimal.h"
-#include "MfMediaFactoryLog.h"
-#include "Misc/Paths.h"
-#include "Modules/ModuleInterface.h"
-#include "Modules/ModuleManager.h"
+#include "MfMediaFactoryPrivate.h"
+
+#include "Containers/Array.h"
+#include "Containers/UnrealString.h"
 #include "IMediaModule.h"
-#include "IMfMediaModule.h"
 #include "IMediaOptions.h"
 #include "IMediaPlayerFactory.h"
+#include "Internationalization/Internationalization.h"
+#include "Misc/Paths.h"
+#include "Modules/ModuleInterface.h"
+#include "Modules/ModuleManager.h"
+#include "UObject/NameTypes.h"
+
+#include "../../MfMedia/Public/IMfMediaModule.h"
 
 
 DEFINE_LOG_CATEGORY(LogMfMediaFactory);
@@ -33,7 +38,7 @@ public:
 
 	//~ IMediaPlayerInfo interface
 
-	virtual bool CanPlayUrl(const FString& Url, const IMediaOptions& Options, TArray<FText>* OutWarnings, TArray<FText>* OutErrors) const override
+	virtual bool CanPlayUrl(const FString& Url, const IMediaOptions* Options, TArray<FText>* OutWarnings, TArray<FText>* OutErrors) const override
 	{
 		FString Scheme;
 		FString Location;
@@ -76,9 +81,9 @@ public:
 		}
 
 		// check options
-		if (OutWarnings != nullptr)
+		if ((OutWarnings != nullptr) && (Options != nullptr))
 		{
-			if (Options.GetMediaOption("PrecacheFile", false) && (Scheme != TEXT("file")))
+			if (Options->GetMediaOption("PrecacheFile", false) && (Scheme != TEXT("file")))
 			{
 				OutWarnings->Add(LOCTEXT("PrecachingNotSupported", "Precaching is supported for local files only"));
 			}
@@ -87,10 +92,10 @@ public:
 		return true;
 	}
 
-	virtual TSharedPtr<IMediaPlayer, ESPMode::ThreadSafe> CreatePlayer() override
+	virtual TSharedPtr<IMediaPlayer, ESPMode::ThreadSafe> CreatePlayer(IMediaEventSink& EventSink) override
 	{
 		auto MfMediaModule = FModuleManager::LoadModulePtr<IMfMediaModule>("MfMedia");
-		return (MfMediaModule != nullptr) ? MfMediaModule->CreatePlayer() : nullptr;
+		return (MfMediaModule != nullptr) ? MfMediaModule->CreatePlayer(EventSink) : nullptr;
 	}
 
 	virtual FText GetDisplayName() const override
@@ -109,6 +114,16 @@ public:
 		return SupportedPlatforms;
 	}
 
+	virtual bool SupportsFeature(EMediaFeature Feature) const override
+	{
+		return ((Feature == EMediaFeature::AudioSamples) ||
+				(Feature == EMediaFeature::AudioTracks) ||
+				(Feature == EMediaFeature::CaptionTracks) ||
+				(Feature == EMediaFeature::OverlaySamples) ||
+				(Feature == EMediaFeature::VideoSamples) ||
+				(Feature == EMediaFeature::VideoTracks));
+	}
+
 public:
 
 	//~ IModuleInterface interface
@@ -145,7 +160,7 @@ public:
 		// supported platforms
 		SupportedPlatforms.Add(TEXT("XboxOne"));
 #if defined(WINVER) && WINVER >= 0x0601
-//		SupportedPlatforms.Add(TEXT("Windows")); // disabled until 4.16, because it's currently broken on Windows
+//		SupportedPlatforms.Add(TEXT("Windows"));
 #endif
 
 		// supported schemes
diff --git a/Engine/Plugins/Media/WmfMedia/Source/WmfMedia/Private/Player/WmfMediaPlayer.cpp b/Engine/Plugins/Media/WmfMedia/Source/WmfMedia/Private/Player/WmfMediaPlayer.cpp
--- a/Engine/Plugins/Media/WmfMedia/Source/WmfMedia/Private/Player/WmfMediaPlayer.cpp
+++ b/Engine/Plugins/Media/WmfMedia/Source/WmfMedia/Private/Player/WmfMediaPlayer.cpp
@@ -1,33 +1,34 @@
 // Copyright 1998-2017 Epic Games, Inc. All Rights Reserved.
 
 #include "WmfMediaPlayer.h"
-#include "WmfMediaResolver.h"
-#include "WmfMediaSession.h"
-#include "WmfMediaSampler.h"
-#include "WmfMediaUtils.h"
-#include "Serialization/ArrayReader.h"
-#include "HAL/FileManager.h"
-#include "Misc/FileHelper.h"
 
 #if WMFMEDIA_SUPPORTED_PLATFORM
 
+#include "Async/Async.h"
+#include "IMediaEventSink.h"
+#include "IMediaOptions.h"
+#include "Misc/Optional.h"
+#include "UObject/Class.h"
+
+#include "WmfMediaSession.h"
+#include "WmfMediaSettings.h"
+#include "WmfMediaTracks.h"
+#include "WmfMediaUtils.h"
+
 #include "AllowWindowsPlatformTypes.h"
 
+
 /* FWmfVideoPlayer structors
  *****************************************************************************/
 
-FWmfMediaPlayer::FWmfMediaPlayer()
+FWmfMediaPlayer::FWmfMediaPlayer(IMediaEventSink& InEventSink)
 	: Duration(0)
+	, EventSink(InEventSink)
+	, Session(new FWmfMediaSession)
+	, Tracks(MakeShared<FWmfMediaTracks, ESPMode::ThreadSafe>())
 {
-	MediaSession = new FWmfMediaSession(EMediaState::Closed);
-	Resolver = new FWmfMediaResolver;
-
-	Tracks.OnSelectionChanged().AddLambda([this]() {
-		if (MediaSession != NULL)
-		{
-			MediaSession->SetTopology(Tracks.CreateTopology());
-		}
-	});
+	check(Session != NULL);
+	check(Tracks.IsValid());
 }
 
 
@@ -42,53 +43,39 @@ FWmfMediaPlayer::~FWmfMediaPlayer()
 
 void FWmfMediaPlayer::Close()
 {
-	Resolver->Cancel();
-
-	if ((MediaSession == NULL) || (MediaSession->GetState() == EMediaState::Closed))
+	if ((Session->GetState() == EMediaState::Closed))
 	{
 		return;
 	}
 
-	// unregister session callbacks
-	MediaSession->OnSessionEvent().RemoveAll(this);
-
-	if (MediaSession->GetState() == EMediaState::Playing)
-	{
-		MediaEvent.Broadcast(EMediaEvent::PlaybackSuspended);
-	}
-
-	// close and release session
-	MediaSession->SetState(EMediaState::Closed);
-	MediaSession = new FWmfMediaSession(EMediaState::Closed);
-
-	// release media source
-	if (MediaSource != NULL)
-	{
-		MediaSource->Shutdown();
-		MediaSource = NULL;
-	}
+	Session->Shutdown();
 
 	// reset player
 	Duration = 0;
-	Info.Empty();
 	MediaUrl = FString();
-	Tracks.Reset();
+	Tracks->Shutdown();
 
 	// notify listeners
-	MediaEvent.Broadcast(EMediaEvent::TracksChanged);
-	MediaEvent.Broadcast(EMediaEvent::MediaClosed);
+	EventSink.ReceiveMediaEvent(EMediaEvent::TracksChanged);
+	EventSink.ReceiveMediaEvent(EMediaEvent::MediaClosed);
+}
+
+
+IMediaCache& FWmfMediaPlayer::GetCache()
+{
+	return *this;
 }
 
 
 IMediaControls& FWmfMediaPlayer::GetControls()
 {
-	return *MediaSession;
+	return *Session;
 }
 
 
 FString FWmfMediaPlayer::GetInfo() const
 {
-	return Info;
+	return Tracks->GetInfo();
 }
 
 
@@ -99,17 +86,16 @@ FName FWmfMediaPlayer::GetName() const
 }
 
 
-IMediaOutput& FWmfMediaPlayer::GetOutput()
+IMediaSamples& FWmfMediaPlayer::GetSamples()
 {
-	return Tracks;
+	return *Tracks;
 }
 
 
 FString FWmfMediaPlayer::GetStats() const
 {
 	FString Result;
-
-	Tracks.AppendStats(Result);
+	Tracks->AppendStats(Result);
 
 	return Result;
 }
@@ -117,7 +103,7 @@ FString FWmfMediaPlayer::GetStats() const
 
 IMediaTracks& FWmfMediaPlayer::GetTracks()
 {
-	return Tracks;
+	return *Tracks;
 }
 
 
@@ -127,7 +113,13 @@ FString FWmfMediaPlayer::GetUrl() const
 }
 
 
-bool FWmfMediaPlayer::Open(const FString& Url, const IMediaOptions& Options)
+IMediaView& FWmfMediaPlayer::GetView()
+{
+	return *this;
+}
+
+
+bool FWmfMediaPlayer::Open(const FString& Url, const IMediaOptions* Options)
 {
 	Close();
 
@@ -136,217 +128,109 @@ bool FWmfMediaPlayer::Open(const FString& Url, const IMediaOptions& Options)
 		return false;
 	}
 
-	// open local files via platform file system
-	bool Resolving = false;
-
-	if (Url.StartsWith(TEXT("file://")))
-	{
-		TSharedPtr<FArchive, ESPMode::ThreadSafe> Archive;
-		const TCHAR* FilePath = &Url[7];
-
-		if (Options.GetMediaOption("PrecacheFile", false))
-		{
-			FArrayReader* Reader = new FArrayReader;
-
-			if (FFileHelper::LoadFileToArray(*Reader, FilePath))
-			{
-				Archive = MakeShareable(Reader);
-			}
-			else
-			{
-				delete Reader;
-			}
-		}
-		else
-		{
-			Archive = MakeShareable(IFileManager::Get().CreateFileReader(FilePath));
-		}
-
-		if (!Archive.IsValid())
-		{
-			UE_LOG(LogWmfMedia, Warning, TEXT("Failed to open media file: %s"), FilePath);
+	const bool Precache = (Options != nullptr) ? Options->GetMediaOption("PrecacheFile", false) : false;
 
-			return false;
-		}
-
-		Resolving = Resolver->ResolveByteStream(Archive.ToSharedRef(), Url, *this);
-	}
-	else
-	{
-		Resolving = Resolver->ResolveUrl(Url, *this);
-	}
-
-	if (Resolving)
-	{
-		MediaSession = new FWmfMediaSession(EMediaState::Preparing);
-	}
-
-	return Resolving;
+	return InitializePlayer(nullptr, Url, Precache);
 }
 
 
-bool FWmfMediaPlayer::Open(const TSharedRef<FArchive, ESPMode::ThreadSafe>& Archive, const FString& OriginalUrl, const IMediaOptions& Options)
+bool FWmfMediaPlayer::Open(const TSharedRef<FArchive, ESPMode::ThreadSafe>& Archive, const FString& OriginalUrl, const IMediaOptions* /*Options*/)
 {
 	Close();
 
-	if ((Archive->TotalSize() == 0) || OriginalUrl.IsEmpty())
+	if (Archive->TotalSize() == 0)
 	{
+		UE_LOG(LogWmfMedia, Error, TEXT("Cannot open media from archive (archive is empty)."));
 		return false;
 	}
 
-	return Resolver->ResolveByteStream(Archive, OriginalUrl, *this);
-}
-
-
-void FWmfMediaPlayer::TickPlayer(float DeltaTime)
-{
-	TFunction<void()> Task;
-
-	while (PlayerTasks.Dequeue(Task))
+	if (OriginalUrl.IsEmpty())
 	{
-		Task();
+		UE_LOG(LogWmfMedia, Error, TEXT("Cannot open media from archive (no original URL provided)."));
+		return false;
 	}
-}
 
-
-void FWmfMediaPlayer::TickVideo(float DeltaTime)
-{
-	// do nothing
+	return InitializePlayer(Archive, OriginalUrl, false);
 }
 
 
-
-/* IWmfMediaResolverCallbacks interface
- *****************************************************************************/
-
-void FWmfMediaPlayer::ProcessResolveComplete(TComPtr<IUnknown> SourceObject, FString ResolvedUrl)
+void FWmfMediaPlayer::TickFetch(FTimespan /*DeltaTime*/)
 {
-	PlayerTasks.Enqueue([=]() {
-		MediaEvent.Broadcast(
-			InitializeMediaSession(SourceObject, ResolvedUrl)
-				? EMediaEvent::MediaOpened
-				: EMediaEvent::MediaOpenFailed
-		);
-	});
-}
+	bool MediaSourceChanged = false;
+	bool TopologyChanged = false;
 
+	Tracks->GetFlags(MediaSourceChanged, TopologyChanged);
 
-void FWmfMediaPlayer::ProcessResolveFailed(FString FailedUrl)
-{
-	PlayerTasks.Enqueue([=]() {
-		MediaSession = new FWmfMediaSession(EMediaState::Closed);
-		MediaEvent.Broadcast(EMediaEvent::MediaOpenFailed);
-	});
-}
-
-
-/* FWmfMediaPlayer implementation
- *****************************************************************************/
-
-bool FWmfMediaPlayer::InitializeMediaSession(IUnknown* SourceObject, const FString& SourceUrl)
-{
-	if (SourceObject == nullptr)
+	if (MediaSourceChanged)
 	{
-		return false;
+		EventSink.ReceiveMediaEvent(EMediaEvent::TracksChanged);
 	}
 
-	UE_LOG(LogWmfMedia, Verbose, TEXT("Initializing media session for %s"), *SourceUrl);
-
-	// create presentation descriptor
-	TComPtr<IMFMediaSource> MediaSourceObject;
+	if (TopologyChanged)
 	{
-		HRESULT Result = SourceObject->QueryInterface(IID_PPV_ARGS(&MediaSourceObject));
-
-		if (FAILED(Result))
+		if (Tracks->IsInitialized())
 		{
-			UE_LOG(LogWmfMedia, Error, TEXT("Failed to query media source: %s"), *WmfMedia::ResultToString(Result));
-			return false;
+			Session->SetTopology(Tracks->CreateTopology(), Tracks->GetDuration());
 		}
-	}
-
-	TComPtr<IMFPresentationDescriptor> PresentationDescriptor;
-	{
-		HRESULT Result = MediaSourceObject->CreatePresentationDescriptor(&PresentationDescriptor);
-
-		if (FAILED(Result))
+		else
 		{
-			UE_LOG(LogWmfMedia, Error, TEXT("Failed to create presentation descriptor: %s"), *WmfMedia::ResultToString(Result));
-			return false;
+			Session->Shutdown();
+			EventSink.ReceiveMediaEvent(EMediaEvent::MediaOpenFailed);
 		}
 	}
 
-	// initialize tracks
-	if (!Tracks.Initialize(*MediaSourceObject, *PresentationDescriptor, Info))
+	if (MediaSourceChanged || TopologyChanged)
 	{
-		UE_LOG(LogWmfMedia, Error, TEXT("Failed to initialize tracks"));
-		return false;
+		Tracks->ClearFlags();
 	}
+}
 
-	// get media duration
-	UINT64 PresentationDuration = 0;
-	PresentationDescriptor->GetUINT64(MF_PD_DURATION, &PresentationDuration);
-	Duration = FTimespan(PresentationDuration);
 
-	// create session
-	MediaUrl = SourceUrl;
-	MediaSource = MediaSourceObject;
-	MediaSession = new FWmfMediaSession(Duration);
+void FWmfMediaPlayer::TickInput(FTimespan DeltaTime)
+{
+	// forward session events
+	TArray<EMediaEvent> OutEvents;
+	Session->GetEvents(OutEvents);
 
-	if (MediaSession->GetState() == EMediaState::Error)
+	for (const auto& Event : OutEvents)
 	{
-		return false;
+		EventSink.ReceiveMediaEvent(Event);
 	}
-
-	MediaSession->OnSessionEvent().AddRaw(this, &FWmfMediaPlayer::HandleSessionEvent);
-	MediaEvent.Broadcast(EMediaEvent::TracksChanged);
-
-	return true;
 }
 
 
-/* FWmfMediaPlayer callbacks
+/* FWmfMediaPlayer implementation
  *****************************************************************************/
 
-void FWmfMediaPlayer::HandleSessionEvent(MediaEventType EventType)
+bool FWmfMediaPlayer::InitializePlayer(const TSharedPtr<FArchive, ESPMode::ThreadSafe>& Archive, const FString& Url, bool Precache)
 {
-	TOptional<EMediaEvent> Event;
+	UE_LOG(LogWmfMedia, Verbose, TEXT("Player %llx: Initializing %s (archive = %s, precache = %s)"), this, *Url, Archive.IsValid() ? TEXT("yes") : TEXT("no"), Precache ? TEXT("yes") : TEXT("no"));
+
+	const auto Settings = GetDefault<UWmfMediaSettings>();
+	check(Settings != nullptr);
 
-	// process event
-	switch (EventType)
+	if (!Session->Initialize(Settings->LowLatency))
 	{
-	case MESessionClosed:
-	case MEError:
-		Tracks.SetPaused(true);
-		Tracks.FlushSinks();
-		break;
-
-	case MESessionStarted:
-		if (!FMath::IsNearlyZero(MediaSession->GetRate()))
-		{
-			Event = EMediaEvent::PlaybackResumed;
-			Tracks.SetPaused(false);
-		}
-		break;
-
-	case MESessionEnded:
-		Event = EMediaEvent::PlaybackEndReached;
-		Tracks.SetPaused(true);
-		Tracks.FlushSinks();
-		break;
-
-	case MESessionStopped:
-		Event = EMediaEvent::PlaybackSuspended;
-		Tracks.SetPaused(true);
-		break;
+		return false;
 	}
 
-	// forward event to game thread
-	if (Event.IsSet())
+	MediaUrl = Url;
+
+	// initialize presentation on a separate thread
+	const EAsyncExecution Execution = Precache ? EAsyncExecution::Thread : EAsyncExecution::ThreadPool;
+
+	Async<void>(Execution, [Archive, Url, Precache, StreamsPtr = TWeakPtr<FWmfMediaTracks, ESPMode::ThreadSafe>(Tracks)]()
 	{
-		PlayerTasks.Enqueue([=]() {
-			MediaEvent.Broadcast(Event.GetValue());
-		});
-	}
+		TSharedPtr<FWmfMediaTracks, ESPMode::ThreadSafe> PinnedStreams = StreamsPtr.Pin();
+
+		if (PinnedStreams.IsValid())
+		{
+			TComPtr<IMFMediaSource> MediaSource = WmfMedia::ResolveMediaSource(Archive, Url, Precache);
+			PinnedStreams->Initialize(MediaSource);
+		}
+	});
+
+	return true;
 }
 
 
diff --git a/Engine/Plugins/Media/WmfMedia/Source/WmfMedia/Private/Player/WmfMediaPlayer.h b/Engine/Plugins/Media/WmfMedia/Source/WmfMedia/Private/Player/WmfMediaPlayer.h
--- a/Engine/Plugins/Media/WmfMedia/Source/WmfMedia/Private/Player/WmfMediaPlayer.h
+++ b/Engine/Plugins/Media/WmfMedia/Source/WmfMedia/Private/Player/WmfMediaPlayer.h
@@ -2,24 +2,21 @@
 
 #pragma once
 
-#include "../WmfMediaPrivate.h"
-#include "WmfMediaPlayer.h"
-#include "IMediaPlayer.h"
-#include "IMediaTracks.h"
-#include "IWmfMediaResolverCallbacks.h"
-#include "WmfMediaTracks.h"
-#include "Containers/Ticker.h"
-#include "Containers/Queue.h"
+#include "WmfMediaPrivate.h"
 
 #if WMFMEDIA_SUPPORTED_PLATFORM
 
+#include "Containers/UnrealString.h"
+#include "IMediaCache.h"
+#include "IMediaPlayer.h"
+#include "IMediaView.h"
+#include "Misc/Timespan.h"
+
 #include "AllowWindowsPlatformTypes.h"
 
-class FWmfMediaResolver;
 class FWmfMediaSession;
-class IMediaControls;
-class IMediaOutput;
-class IMediaTracks;
+class FWmfMediaTracks;
+class IMediaEventSink;
 
 
 /**
@@ -27,93 +24,71 @@ class IMediaTracks;
  */
 class FWmfMediaPlayer
 	: public IMediaPlayer
-	, protected IWmfMediaResolverCallbacks
+	, protected IMediaCache
+	, protected IMediaView
 {
 public:
 
-	/** Default constructor. */
-	FWmfMediaPlayer();
+	/**
+	 * Create and initialize a new instance.
+	 *
+	 * @param InEventSink The object that receives media events from this player.
+	 */
+	FWmfMediaPlayer(IMediaEventSink& InEventSink);
 
-	/** Destructor. */
-	~FWmfMediaPlayer();
+	/** Virtual destructor. */
+	virtual ~FWmfMediaPlayer();
 
 public:
 
 	//~ IMediaPlayer interface
 
 	virtual void Close() override;
+	virtual IMediaCache& GetCache() override;
 	virtual IMediaControls& GetControls() override;
 	virtual FString GetInfo() const override;
 	virtual FName GetName() const override;
-	virtual IMediaOutput& GetOutput() override;
+	virtual IMediaSamples& GetSamples() override;
 	virtual FString GetStats() const override;
 	virtual IMediaTracks& GetTracks() override;
 	virtual FString GetUrl() const override;
-	virtual bool Open(const FString& Url, const IMediaOptions& Options) override;
-	virtual bool Open(const TSharedRef<FArchive, ESPMode::ThreadSafe>& Archive, const FString& OriginalUrl, const IMediaOptions& Options) override;
-	virtual void TickPlayer(float DeltaTime) override;
-	virtual void TickVideo(float DeltaTime) override;
-	
-	DECLARE_DERIVED_EVENT(FWmfMediaPlayer, IMediaPlayer::FOnMediaEvent, FOnMediaEvent);
-	virtual FOnMediaEvent& OnMediaEvent() override
-	{
-		return MediaEvent;
-	}
-
-protected:
-
-	//~ IWmfMediaResolverCallbacks interface
-
-	virtual void ProcessResolveComplete(TComPtr<IUnknown> SourceObject, FString ResolvedUrl) override;
-	virtual void ProcessResolveFailed(FString FailedUrl) override;
+	virtual IMediaView& GetView() override;
+	virtual bool Open(const FString& Url, const IMediaOptions* Options) override;
+	virtual bool Open(const TSharedRef<FArchive, ESPMode::ThreadSafe>& Archive, const FString& OriginalUrl, const IMediaOptions* Options) override;
+	virtual void TickFetch(FTimespan DeltaTime) override;
+	virtual void TickInput(FTimespan DeltaTime) override;
 
 protected:
 
 	/**
-	 * Initializes the media session for the given media source.
+	 * Initialize the native AvPlayer instance.
 	 *
-	 * @param SourceObject The media source object.
-	 * @param SourceUrl The original URL of the media source.
+	 * @param Archive The archive being used as a media source (optional).
+	 * @param Url The media URL being opened.
+	 * @param Precache Whether to precache media into RAM if InURL is a local file.
 	 * @return true on success, false otherwise.
 	 */
-	bool InitializeMediaSession(IUnknown* SourceObject, const FString& SourceUrl);
-
-private:
-
-	/** Handles session events. */
-	void HandleSessionEvent(MediaEventType EventType);
+	bool InitializePlayer(const TSharedPtr<FArchive, ESPMode::ThreadSafe>& Archive, const FString& Url, bool Precache);
 
 private:
 
 	/** The duration of the currently loaded media. */
 	FTimespan Duration;
 
-	/** Media information string. */
-	FString Info;
-
-	/** Tasks to execute on the player thread. */
-	TQueue<TFunction<void()>> PlayerTasks;
-
-	/** Event delegate that is invoked when a media event occurred. */
-	FOnMediaEvent MediaEvent;
-
-	/** Asynchronous callback object for the media stream. */
-	TComPtr<FWmfMediaSession> MediaSession;
-
-	/** Pointer to the media source object. */
-	TComPtr<IMFMediaSource> MediaSource;
+	/** The media event handler. */
+	IMediaEventSink& EventSink;
 
 	/** The URL of the currently opened media. */
 	FString MediaUrl;
 
-	/** The media source resolver. */
-	TComPtr<FWmfMediaResolver> Resolver;
+	/** Asynchronous callback object for the media session. */
+	TComPtr<FWmfMediaSession> Session;
 
-	/** Track manager. */
-	FWmfMediaTracks Tracks;
+	/** Media streams collection. */
+	TSharedPtr<FWmfMediaTracks, ESPMode::ThreadSafe> Tracks;
 };
 
 
 #include "HideWindowsPlatformTypes.h"
 
-#endif
\ No newline at end of file
+#endif //WMFMEDIA_SUPPORTED_PLATFORM
\ No newline at end of file
