Prevented MfMedia from having a dependency on MFReadWrite and initialzing in server builds

(Fixes server failing on launch in OrionTests)

#tests ran server on build machine, compiled client

diff --git a/Engine/Plugins/Media/MfMedia/Source/MfMedia/MfMedia.Build.cs b/Engine/Plugins/Media/MfMedia/Source/MfMedia/MfMedia.Build.cs
--- a/Engine/Plugins/Media/MfMedia/Source/MfMedia/MfMedia.Build.cs
+++ b/Engine/Plugins/Media/MfMedia/Source/MfMedia/MfMedia.Build.cs
@@ -33,16 +33,19 @@ namespace UnrealBuildTool.Rules
 					"MfMedia/Private/Player",
 				});
 
-			PublicAdditionalLibraries.Add("mfplat.lib");
-			PublicAdditionalLibraries.Add("mfreadwrite.lib");
-			PublicAdditionalLibraries.Add("mfuuid.lib");
-
-			if ((Target.Platform == UnrealTargetPlatform.Win32) ||
-				(Target.Platform == UnrealTargetPlatform.Win64))
+			if (Target.Type != TargetType.Server)
 			{
-				PublicAdditionalLibraries.Add("mf.lib");
-				PublicAdditionalLibraries.Add("Propsys.lib");
-				PublicAdditionalLibraries.Add("shlwapi.lib");
+				PublicAdditionalLibraries.Add("mfplat.lib");
+				PublicAdditionalLibraries.Add("mfreadwrite.lib");
+				PublicAdditionalLibraries.Add("mfuuid.lib");
+
+				if ((Target.Platform == UnrealTargetPlatform.Win32) 
+					|| (Target.Platform == UnrealTargetPlatform.Win64))
+				{
+					PublicAdditionalLibraries.Add("mf.lib");
+					PublicAdditionalLibraries.Add("Propsys.lib");
+					PublicAdditionalLibraries.Add("shlwapi.lib");
+				}
 			}
 		}
 	}
diff --git a/Engine/Plugins/Media/MfMedia/Source/MfMedia/Private/MfMediaPrivate.h b/Engine/Plugins/Media/MfMedia/Source/MfMedia/Private/MfMediaPrivate.h
--- a/Engine/Plugins/Media/MfMedia/Source/MfMedia/Private/MfMediaPrivate.h
+++ b/Engine/Plugins/Media/MfMedia/Source/MfMedia/Private/MfMediaPrivate.h
@@ -4,7 +4,7 @@
 
 #include "Logging/LogMacros.h"
 
-#define MFMEDIA_SUPPORTED_PLATFORM (PLATFORM_XBOXONE || (PLATFORM_WINDOWS && WINVER >= 0x0601 /*Win7*/))
+#define MFMEDIA_SUPPORTED_PLATFORM (PLATFORM_XBOXONE || (PLATFORM_WINDOWS && WINVER >= 0x0601 /*Win7*/ && !UE_SERVER))
 
 #if MFMEDIA_SUPPORTED_PLATFORM
 	#if PLATFORM_WINDOWS
