Copying //UE4/Dev-Sequencer to //UE4/Main (Source: //UE4/Dev-Sequencer @ 3631098)

Change 3618690 on 2017/08/30 by [REDACTED_PERSON]

	Media: Pass-through of global timecode in tick

Change 3620659 on 2017/08/31 by [REDACTED_PERSON]

	WmfMedia: Added string conversion for stream marker types; updated documentation

Change 3620663 on 2017/08/31 by [REDACTED_PERSON]

	MediaUtils: Always queuing up media events to prevent application from interfering with player state during event callbacks

Change 3620878 on 2017/08/31 by [REDACTED_PERSON]

	Media: Using correct parameter for printing pointers

Change 3621327 on 2017/08/31 by [REDACTED_PERSON]

	Core: Added TQueue.Pop

Change 3621419 on 2017/08/31 by [REDACTED_PERSON]

	Skip some mediaplayer calls if not prepared and check for completion in stopped state

Change 3622023 on 2017/09/01 by [REDACTED_PERSON]

	Sequencer: Fix regression where spawnables don't play because they're defined as temporary editor actors.

Change 3622433 on 2017/09/01 by [REDACTED_PERSON]

	Reorder Android media player suspend and playbackendreached events so stop is first

Change 3622568 on 2017/09/01 by [REDACTED_PERSON]

	Media Player: Move sinks and player guid to PostInitProperties instead of PostLoad to fix issues with newly created players.

Change 3622678 on 2017/09/01 by [REDACTED_PERSON]

	Actor Sequence: Move runtime to plugins.

Change 3623243 on 2017/09/01 by [REDACTED_PERSON]

	MfMedia: Resetting media characteristics before trying to reinitialize them

Change 3623246 on 2017/09/01 by [REDACTED_PERSON]

	WmfMedia: Fixed media characteristics not updated when opening media without tracks selected

Change 3623675 on 2017/09/01 by [REDACTED_PERSON]

	MediaUtils: Fixed multichannel downmixing

Change 3623747 on 2017/09/01 by [REDACTED_PERSON]

	WmfMedia: Removed still images from list of supported streams as they require special handling that is currently not implemented

Change 3623757 on 2017/09/01 by [REDACTED_PERSON]

	WmfMedia: Added missing string conversion for Playing state

Change 3623793 on 2017/09/01 by [REDACTED_PERSON]

	Sequencer: Fix ensure when deleting sections. Sections were getting removed from the array during iteration.

Change 3623797 on 2017/09/01 by [REDACTED_PERSON]

	Sequencer: Add transactions to easing type and setting changes.

Change 3623924 on 2017/09/02 by [REDACTED_PERSON]

	Media: Renamed IMediaPlayer::GetName to GetPlayerName

Change 3623931 on 2017/09/02 by [REDACTED_PERSON]

	Media: Added IMediaPlayer::GetMediaName

Change 3623962 on 2017/09/02 by [REDACTED_PERSON]

	WmfMedia: Consistent logging format

Change 3623963 on 2017/09/02 by [REDACTED_PERSON]

	MfMedia: Consistent logging format

Change 3624190 on 2017/09/03 by [REDACTED_PERSON]

	MediaUtils: Made sink overflow log messages a compile option

Change 3624412 on 2017/09/04 by [REDACTED_PERSON]

	Trivial fix for Mac Audio Mixer playback through engine using current AV implementation.  This fixes up what is already there to get us back to a working Mac Media playback via engine.  Although it is still disabled just in case this stream is at a time/feature critical point.  Just remove the && 0  in the AUDIO_PLAYBACK_VIA_ENGINE define at the top of AvfMediaTracks.cpp.

Change 3624769 on 2017/09/04 by [REDACTED_PERSON]

	WmfMedia: Fixed some subtle issues with session state changes & track switching

Change 3624811 on 2017/09/04 by [REDACTED_PERSON]

	WmfMedia: Not compiling WMF related libraries on dedicated server

Change 3624813 on 2017/09/04 by [REDACTED_PERSON]

	WmfMedia: Added utility function to copy MF attributes

Change 3624941 on 2017/09/05 by [REDACTED_PERSON]

	Sequencer: Added slip functionality for sections. Hold down shift and drag the section left and right (as if you were moving it).

Change 3625131 on 2017/09/05 by [REDACTED_PERSON]

	Media: Fixed optional parameter name not being loaded properly for external texture uniform expressions

Change 3625547 on 2017/09/05 by [REDACTED_PERSON]

	AudioMixer: Fixed crash due to legacy int16 audio buffer not initialized

Change 3625569 on 2017/09/05 by [REDACTED_PERSON]

	Cine Camera: Added GetFilmbackPresetName(), SetFilmbackPresetByName(), GetLensPresetName(), SetLensPresetByName().

Change 3625668 on 2017/09/05 by [REDACTED_PERSON]

	Sequence Recorder: Add drag and drop functionality to sequence recorder

Change 3625738 on 2017/09/05 by [REDACTED_PERSON]

	WmfMedia: showing possible image streams that may be incorrectly returned as video streams in media info

Change 3625751 on 2017/09/05 by [REDACTED_PERSON]

	WmfMedia: Restarting session from beginning on track switching if media doesn't support seeking

Change 3625869 on 2017/09/05 by [REDACTED_PERSON]

	Sequencer: Fixed crash caused by manipulating selection while iterating

Change 3625884 on 2017/09/05 by [REDACTED_PERSON]

	Sequencer: Optimization pass on key and trajectory rendering when dealing with large numbers of tracks and/or keys
	  - Memory is reserved where possible on hot code paths
	  - Reduced amount of sorting and shuffling of arrays where possible
	  - Key collections (and thus key groups) are now cacheable so we don't have to regenerate them every frame
	  - Reduced amount of overdraw on keys in sequencer UI
	  - Added editor preference to limit the number of keys shown around the current time when rendering trajectories in the editor (default=250)
	  - Optimized various algorithms used for key introspection and manipulation
	  - Overlapping keys are now highlighted with a red border

Change 3626270 on 2017/09/05 by [REDACTED_PERSON]

	WmfMedia: Added utility function to dump out media type attributes

Change 3626274 on 2017/09/05 by [REDACTED_PERSON]

	WmfMedia: Dumping media type attributes in very verbose debug log

Change 3626275 on 2017/09/05 by [REDACTED_PERSON]

	WmfMedia: Clearing pending topologies before setting new one

Change 3626440 on 2017/09/05 by [REDACTED_PERSON]

	WmfMedia: Resolving stream format types and interlace modes in attribute dump

Change 3626719 on 2017/09/05 by [REDACTED_PERSON]

	MediaAssets: Exposed the OnTracksChanged event on media player assets

Change 3626806 on 2017/09/05 by [REDACTED_PERSON]

	Fix error returns from rungradle.bat and ant.bat

	Change 3626915 on 2017/09/05 by [REDACTED_PERSON]

	Only use immersive mode if enabled on Android

Change 3626980 on 2017/09/05 by [REDACTED_PERSON]

	MediaAssets: Renamed capture device enumeration filters from Unspecified to Unknown

Change 3627333 on 2017/09/05 by [REDACTED_PERSON]

	MediaUtils: Added utility functions for converting common enum values to string

Change 3627702 on 2017/09/06 by [REDACTED_PERSON]

	Sequencer: Fixed crash caused by erroneous persistent references to FSequencer
	  - Also added an ensure to FSequencer::Tick to ensure this doesn't happen again

Change 3628000 on 2017/09/06 by [REDACTED_PERSON]

	Sequencer: Changed ensure for a message log warning

Change 3628548 on 2017/09/06 by [REDACTED_PERSON]

	MediaUtils: Logging media events in player facade

Change 3629189 on 2017/09/06 by [REDACTED_PERSON]

	WmfMedia: Selecting default format and track when initializing media source

Change 3629207 on 2017/09/06 by [REDACTED_PERSON]

	MfMedia: Selecting default format when initializing media source

Change 3629256 on 2017/09/06 by [REDACTED_PERSON]

	MediaUtils: Fixed sample queue count not decremented on Pop

diff --git a/Engine/Build/Android/Java/src/com/epicgames/ue4/GameActivity.java b/Engine/Build/Android/Java/src/com/epicgames/ue4/GameActivity.java
--- a/Engine/Build/Android/Java/src/com/epicgames/ue4/GameActivity.java
+++ b/Engine/Build/Android/Java/src/com/epicgames/ue4/GameActivity.java
@@ -1102,8 +1102,11 @@ public class GameActivity extends NativeActivity implements SurfaceHolder.Callba
 	{
 		if(keyCode == KeyEvent.KEYCODE_BACK ||keyCode == KeyEvent.KEYCODE_VOLUME_DOWN || keyCode == KeyEvent.KEYCODE_VOLUME_UP)
 		{
-			Log.debug("=== Restoring Transparent Bars due to KeyCode ===");
-			restoreTranslucentBarsDelayed();
+			if (ShouldHideUI)
+			{
+				Log.debug("=== Restoring Transparent Bars due to KeyCode ===");
+				restoreTranslucentBarsDelayed();
+			}
 		}
 
 		return super.onKeyDown(keyCode, event);
diff --git a/Engine/Build/Android/Java/src/com/epicgames/ue4/MediaPlayer14.java b/Engine/Build/Android/Java/src/com/epicgames/ue4/MediaPlayer14.java
--- a/Engine/Build/Android/Java/src/com/epicgames/ue4/MediaPlayer14.java
+++ b/Engine/Build/Android/Java/src/com/epicgames/ue4/MediaPlayer14.java
@@ -96,6 +96,7 @@ public class MediaPlayer14
 		setOnPreparedListener(new MediaPlayer.OnPreparedListener() {
 			@Override
 			public void onPrepared(MediaPlayer player) {
+//				GameActivity.Log.debug("*** MEDIA PREPARED ***");
 				synchronized(player)
 				{
 					Prepared = true;
@@ -108,6 +109,7 @@ public class MediaPlayer14
 			public void onCompletion(MediaPlayer player) {
 				synchronized(player)
 				{
+//					GameActivity.Log.debug("*** MEDIA COMPLETION ***");
 					if (Looping)
 					{
 						seekTo(0);
@@ -421,7 +423,10 @@ public class MediaPlayer14
 		synchronized(this)
 		{
 			Completed = false;
-			super.start();
+			if (Prepared)
+			{
+				super.start();
+			}
 		}
 	}
 
@@ -430,8 +435,26 @@ public class MediaPlayer14
 		synchronized(this)
 		{
 			Completed = false;
-			super.stop();
+			if (Prepared)
+			{
+				super.stop();
+			}
+		}
+	}
+
+	public int getCurrentPosition()
+	{
+		int position = 0;
+
+		synchronized(this)
+		{
+			if (Prepared)
+			{
+				position = super.getCurrentPosition();
+			}
 		}
+
+		return position;
 	}
 
 	public void seekTo(int position)
@@ -439,7 +462,10 @@ public class MediaPlayer14
 		synchronized (this)
 		{
 			Completed = false;
-			super.seekTo(position);
+			if (Prepared)
+			{
+				super.seekTo(position);
+			}
 		}
 	}
 
diff --git a/Engine/Plugins/Media/AndroidCamera/Source/AndroidCamera/Private/Player/AndroidCameraPlayer.cpp b/Engine/Plugins/Media/AndroidCamera/Source/AndroidCamera/Private/Player/AndroidCameraPlayer.cpp
--- a/Engine/Plugins/Media/AndroidCamera/Source/AndroidCamera/Private/Player/AndroidCameraPlayer.cpp
+++ b/Engine/Plugins/Media/AndroidCamera/Source/AndroidCamera/Private/Player/AndroidCameraPlayer.cpp
@@ -164,7 +164,7 @@ FString FAndroidCameraPlayer::GetInfo() const
 }
 
 
-FName FAndroidCameraPlayer::GetName() const
+FName FAndroidCameraPlayer::GetPlayerName() const
 {
 	static FName PlayerName(TEXT("AndroidCamera"));
 	return PlayerName;
@@ -267,7 +267,7 @@ bool FAndroidCameraPlayer::Open(const FString& Url, const IMediaOptions* /*Optio
 
 bool FAndroidCameraPlayer::Open(const TSharedRef<FArchive, ESPMode::ThreadSafe>& Archive, const FString& OriginalUrl, const IMediaOptions* /*Options*/)
 {
-	return false; // @todo AndroidCamera: implement opening media from FArchive
+	return false; // not supported
 }
 
 
@@ -281,7 +281,7 @@ void FAndroidCameraPlayer::SetGuid(const FGuid& Guid)
 }
 
 
-void FAndroidCameraPlayer::TickFetch(FTimespan DeltaTime)
+void FAndroidCameraPlayer::TickFetch(FTimespan DeltaTime, FTimespan Timecode)
 {
 	if (CurrentState != EMediaState::Playing && CurrentState != EMediaState::Paused)
 	{
@@ -576,7 +576,7 @@ void FAndroidCameraPlayer::TickFetch(FTimespan DeltaTime)
 }
 
 
-void FAndroidCameraPlayer::TickInput(FTimespan DeltaTime)
+void FAndroidCameraPlayer::TickInput(FTimespan DeltaTime, FTimespan Timecode)
 {
 	if (CurrentState != EMediaState::Playing)
 	{
@@ -999,11 +999,11 @@ bool FAndroidCameraPlayer::SelectTrack(EMediaTrackType TrackType, int32 TrackInd
 	case EMediaTrackType::Audio:
 		if (TrackIndex != SelectedAudioTrack)
 		{
-			UE_LOG(LogAndroidCamera, Verbose, TEXT("Player %llx: Selecting audio track %i instead of %i (%i tracks)"), this, TrackIndex, SelectedVideoTrack, AudioTracks.Num());
+			UE_LOG(LogAndroidCamera, Verbose, TEXT("Player %p: Selecting audio track %i instead of %i (%i tracks)"), this, TrackIndex, SelectedVideoTrack, AudioTracks.Num());
 
 			if (TrackIndex == INDEX_NONE)
 			{
-				UE_LOG(LogAndroidCamera, VeryVerbose, TEXT("Player %llx: Disabling audio"), this, TrackIndex);
+				UE_LOG(LogAndroidCamera, VeryVerbose, TEXT("Player %p: Disabling audio"), this, TrackIndex);
 
 				JavaCameraPlayer->SetAudioEnabled(false);
 			}
@@ -1014,7 +1014,7 @@ bool FAndroidCameraPlayer::SelectTrack(EMediaTrackType TrackType, int32 TrackInd
 					return false;
 				}
 
-				UE_LOG(LogAndroidCamera, VeryVerbose, TEXT("Player %llx: Enabling audio"), this, TrackIndex);
+				UE_LOG(LogAndroidCamera, VeryVerbose, TEXT("Player %p: Enabling audio"), this, TrackIndex);
 
 				JavaCameraPlayer->SetAudioEnabled(true);
 			}
@@ -1026,11 +1026,11 @@ bool FAndroidCameraPlayer::SelectTrack(EMediaTrackType TrackType, int32 TrackInd
 	case EMediaTrackType::Caption:
 		if (TrackIndex != SelectedCaptionTrack)
 		{
-			UE_LOG(LogAndroidCamera, Verbose, TEXT("Player %llx: Selecting caption track %i instead of %i (%i tracks)"), this, TrackIndex, SelectedCaptionTrack, CaptionTracks.Num());
+			UE_LOG(LogAndroidCamera, Verbose, TEXT("Player %p: Selecting caption track %i instead of %i (%i tracks)"), this, TrackIndex, SelectedCaptionTrack, CaptionTracks.Num());
 
 			if (TrackIndex == INDEX_NONE)
 			{
-				UE_LOG(LogAndroidCamera, VeryVerbose, TEXT("Player %llx: Disabling captions"), this, TrackIndex);
+				UE_LOG(LogAndroidCamera, VeryVerbose, TEXT("Player %p: Disabling captions"), this, TrackIndex);
 			}
 			else
 			{
@@ -1039,7 +1039,7 @@ bool FAndroidCameraPlayer::SelectTrack(EMediaTrackType TrackType, int32 TrackInd
 					return false;
 				}
 
-				UE_LOG(LogAndroidCamera, VeryVerbose, TEXT("Player %llx: Enabling captions"), this, TrackIndex);
+				UE_LOG(LogAndroidCamera, VeryVerbose, TEXT("Player %p: Enabling captions"), this, TrackIndex);
 			}
 
 			SelectedCaptionTrack = TrackIndex;
@@ -1049,11 +1049,11 @@ bool FAndroidCameraPlayer::SelectTrack(EMediaTrackType TrackType, int32 TrackInd
 	case EMediaTrackType::Video:
 		if (TrackIndex != SelectedVideoTrack)
 		{
-			UE_LOG(LogAndroidCamera, Verbose, TEXT("Player %llx: Selecting video track %i instead of %i (%i tracks)."), this, TrackIndex, SelectedVideoTrack, VideoTracks.Num());
+			UE_LOG(LogAndroidCamera, Verbose, TEXT("Player %p: Selecting video track %i instead of %i (%i tracks)."), this, TrackIndex, SelectedVideoTrack, VideoTracks.Num());
 
 			if (TrackIndex == INDEX_NONE)
 			{
-				UE_LOG(LogAndroidCamera, VeryVerbose, TEXT("Player %llx: Disabling video"), this, TrackIndex);
+				UE_LOG(LogAndroidCamera, VeryVerbose, TEXT("Player %p: Disabling video"), this, TrackIndex);
 				JavaCameraPlayer->SetVideoEnabled(false);
 			}
 			else
@@ -1078,7 +1078,7 @@ bool FAndroidCameraPlayer::SelectTrack(EMediaTrackType TrackType, int32 TrackInd
 					SetRate(OldRate);
 				}
 
-				UE_LOG(LogAndroidCamera, VeryVerbose, TEXT("Player %llx: Enabling video"), this, TrackIndex);
+				UE_LOG(LogAndroidCamera, VeryVerbose, TEXT("Player %p: Enabling video"), this, TrackIndex);
 				JavaCameraPlayer->SetVideoEnabled(true);
 			}
 
@@ -1279,7 +1279,7 @@ bool FAndroidCameraPlayer::Seek(const FTimespan& Time)
 		return false;
 	}
 
-	UE_LOG(LogAndroidCamera, Verbose, TEXT("Player %llx: Seeking to %s"), this, *Time.ToString());
+	UE_LOG(LogAndroidCamera, Verbose, TEXT("Player %p: Seeking to %s"), this, *Time.ToString());
 
 	JavaCameraPlayer->SeekTo(static_cast<int32>(Time.GetTotalMilliseconds()));
 	EventSink.ReceiveMediaEvent(EMediaEvent::SeekCompleted);
@@ -1318,7 +1318,7 @@ bool FAndroidCameraPlayer::SetRate(float Rate)
 		return true; // rate already set
 	}
 
-	UE_LOG(LogAndroidCamera, Verbose, TEXT("Player %llx: Setting rate from to %f to %f"), this, GetRate(), Rate);
+	UE_LOG(LogAndroidCamera, Verbose, TEXT("Player %p: Setting rate from to %f to %f"), this, GetRate(), Rate);
 
 	if (Rate == 0.0f)
 	{
diff --git a/Engine/Plugins/Media/AndroidCamera/Source/AndroidCamera/Private/Player/AndroidCameraPlayer.h b/Engine/Plugins/Media/AndroidCamera/Source/AndroidCamera/Private/Player/AndroidCameraPlayer.h
--- a/Engine/Plugins/Media/AndroidCamera/Source/AndroidCamera/Private/Player/AndroidCameraPlayer.h
+++ b/Engine/Plugins/Media/AndroidCamera/Source/AndroidCamera/Private/Player/AndroidCameraPlayer.h
@@ -51,7 +51,7 @@ public:
 	virtual IMediaCache& GetCache() override;
 	virtual IMediaControls& GetControls() override;
 	virtual FString GetInfo() const override;
-	virtual FName GetName() const override;
+	virtual FName GetPlayerName() const override;
 	virtual IMediaSamples& GetSamples() override;
 	virtual FString GetStats() const override;
 	virtual IMediaTracks& GetTracks() override;
@@ -60,8 +60,8 @@ public:
 	virtual bool Open(const FString& Url, const IMediaOptions* Options) override;
 	virtual bool Open(const TSharedRef<FArchive, ESPMode::ThreadSafe>& Archive, const FString& OriginalUrl, const IMediaOptions* Options) override;
 	virtual void SetGuid(const FGuid& Guid) override;
-	virtual void TickFetch(FTimespan DeltaTime) override;
-	virtual void TickInput(FTimespan DeltaTime) override;
+	virtual void TickFetch(FTimespan DeltaTime, FTimespan Timecode) override;
+	virtual void TickInput(FTimespan DeltaTime, FTimespan Timecode) override;
 
 protected:
 
diff --git a/Engine/Plugins/Media/AndroidCamera/Source/AndroidCameraFactory/Private/AndroidCameraFactoryModule.cpp b/Engine/Plugins/Media/AndroidCamera/Source/AndroidCameraFactory/Private/AndroidCameraFactoryModule.cpp
--- a/Engine/Plugins/Media/AndroidCamera/Source/AndroidCameraFactory/Private/AndroidCameraFactoryModule.cpp
+++ b/Engine/Plugins/Media/AndroidCamera/Source/AndroidCameraFactory/Private/AndroidCameraFactoryModule.cpp
@@ -88,7 +88,7 @@ public:
 		return LOCTEXT("MediaPlayerDisplayName", "Android Camera");
 	}
 
-	virtual FName GetName() const override
+	virtual FName GetPlayerName() const override
 	{
 		static FName PlayerName(TEXT("AndroidCamera"));
 		return PlayerName;
diff --git a/Engine/Plugins/Media/AndroidMedia/Source/AndroidMedia/Private/Player/AndroidMediaPlayer.cpp b/Engine/Plugins/Media/AndroidMedia/Source/AndroidMedia/Private/Player/AndroidMediaPlayer.cpp
--- a/Engine/Plugins/Media/AndroidMedia/Source/AndroidMedia/Private/Player/AndroidMediaPlayer.cpp
+++ b/Engine/Plugins/Media/AndroidMedia/Source/AndroidMedia/Private/Player/AndroidMediaPlayer.cpp
@@ -104,6 +104,34 @@ FAndroidMediaPlayer::~FAndroidMediaPlayer()
 /* IMediaPlayer interface
  *****************************************************************************/
 
+#if ANDROIDMEDIAPLAYER_USE_NATIVELOGGING
+// Useful for debugging
+static void DumpState(EMediaState state)
+{
+	switch (state)
+	{
+	case EMediaState::Closed:
+		FPlatformMisc::LowLevelOutputDebugString(TEXT("AndroidMediaPlayer: CurrentState = Closed"));
+		break;
+	case EMediaState::Error:
+		FPlatformMisc::LowLevelOutputDebugString(TEXT("AndroidMediaPlayer: CurrentState = Error"));
+		break;
+	case EMediaState::Paused:
+		FPlatformMisc::LowLevelOutputDebugString(TEXT("AndroidMediaPlayer: CurrentState = Paused"));
+		break;
+	case EMediaState::Playing:
+		FPlatformMisc::LowLevelOutputDebugString(TEXT("AndroidMediaPlayer: CurrentState = Playing"));
+		break;
+	case EMediaState::Preparing:
+		FPlatformMisc::LowLevelOutputDebugString(TEXT("AndroidMediaPlayer: CurrentState = Preparing"));
+		break;
+	case EMediaState::Stopped:
+		FPlatformMisc::LowLevelOutputDebugString(TEXT("AndroidMediaPlayer: CurrentState = Stopped"));
+		break;
+	}
+}
+#endif
+
 void FAndroidMediaPlayer::Close()
 {
 	#if ANDROIDMEDIAPLAYER_USE_NATIVELOGGING
@@ -164,7 +192,7 @@ FString FAndroidMediaPlayer::GetInfo() const
 }
 
 
-FName FAndroidMediaPlayer::GetName() const
+FName FAndroidMediaPlayer::GetPlayerName() const
 {
 	static FName PlayerName(TEXT("AndroidMedia"));
 	return PlayerName;
@@ -306,7 +334,7 @@ void FAndroidMediaPlayer::SetGuid(const FGuid& Guid)
 }
 
 
-void FAndroidMediaPlayer::TickFetch(FTimespan DeltaTime)
+void FAndroidMediaPlayer::TickFetch(FTimespan DeltaTime, FTimespan /*Timecode*/)
 {
 	if (CurrentState != EMediaState::Playing && CurrentState != EMediaState::Paused)
 	{
@@ -579,7 +607,7 @@ void FAndroidMediaPlayer::TickFetch(FTimespan DeltaTime)
 }
 
 
-void FAndroidMediaPlayer::TickInput(FTimespan DeltaTime)
+void FAndroidMediaPlayer::TickInput(FTimespan DeltaTime, FTimespan /*Timecode*/)
 {
 	if (CurrentState != EMediaState::Playing)
 	{
@@ -596,6 +624,11 @@ void FAndroidMediaPlayer::TickInput(FTimespan DeltaTime)
 			PauseHandle.Reset();
 		}
 
+		if (!JavaMediaPlayer.IsValid())
+		{
+			return;
+		}
+
 #if ANDROIDMEDIAPLAYER_USE_PREPAREASYNC
 		// if preparing, see if finished
 		if (CurrentState == EMediaState::Preparing)
@@ -610,7 +643,19 @@ void FAndroidMediaPlayer::TickInput(FTimespan DeltaTime)
 				InitializePlayer();
 			}
 		}
+		else
+#endif
+		if (CurrentState == EMediaState::Stopped)
+		{
+			if (JavaMediaPlayer->DidComplete())
+			{
+				EventSink.ReceiveMediaEvent(EMediaEvent::PlaybackEndReached);
+
+#if ANDROIDMEDIAPLAYER_USE_NATIVELOGGING
+				FPlatformMisc::LowLevelOutputDebugStringf(TEXT("FAndroidMedia::Tick - PlaybackEndReached - stopped - %s"), *PlayerGuid.ToString());
 #endif
+			}
+		}
 
 		return;
 	}
@@ -633,23 +678,23 @@ void FAndroidMediaPlayer::TickInput(FTimespan DeltaTime)
 	// generate events
 	if (!JavaMediaPlayer->IsPlaying())
 	{
-		if (JavaMediaPlayer->DidComplete())
-		{
-			EventSink.ReceiveMediaEvent(EMediaEvent::PlaybackEndReached);
-
-			#if ANDROIDMEDIAPLAYER_USE_NATIVELOGGING
-				FPlatformMisc::LowLevelOutputDebugStringf(TEXT("FAndroidMedia::Tick - PlaybackEndReached - !playing - %s"), *PlayerGuid.ToString());
-			#endif
-		}
-
 		// might catch it restarting the loop so ignore if looping
 		if (!bLooping)
 		{
 			CurrentState = EMediaState::Stopped;
 			EventSink.ReceiveMediaEvent(EMediaEvent::PlaybackSuspended);
 
+#if ANDROIDMEDIAPLAYER_USE_NATIVELOGGING
+			FPlatformMisc::LowLevelOutputDebugStringf(TEXT("FAndroidMedia::Tick - PlaybackSuspended - !playing - %s"), *PlayerGuid.ToString());
+#endif
+		}
+
+		if (JavaMediaPlayer->DidComplete())
+		{
+			EventSink.ReceiveMediaEvent(EMediaEvent::PlaybackEndReached);
+
 			#if ANDROIDMEDIAPLAYER_USE_NATIVELOGGING
-				FPlatformMisc::LowLevelOutputDebugStringf(TEXT("FAndroidMedia::Tick - PlaybackSuspended - !playing - %s"), *PlayerGuid.ToString());
+				FPlatformMisc::LowLevelOutputDebugStringf(TEXT("FAndroidMedia::Tick - PlaybackEndReached - !playing - %s"), *PlayerGuid.ToString());
 			#endif
 		}
 	}
@@ -949,11 +994,11 @@ bool FAndroidMediaPlayer::SelectTrack(EMediaTrackType TrackType, int32 TrackInde
 	case EMediaTrackType::Audio:
 		if (TrackIndex != SelectedAudioTrack)
 		{
-			UE_LOG(LogAndroidMedia, Verbose, TEXT("Player %llx: Selecting audio track %i instead of %i (%i tracks)"), this, TrackIndex, SelectedVideoTrack, AudioTracks.Num());
+			UE_LOG(LogAndroidMedia, Verbose, TEXT("Player %p: Selecting audio track %i instead of %i (%i tracks)"), this, TrackIndex, SelectedVideoTrack, AudioTracks.Num());
 
 			if (TrackIndex == INDEX_NONE)
 			{
-				UE_LOG(LogAndroidMedia, VeryVerbose, TEXT("Player %llx: Disabling audio"), this, TrackIndex);
+				UE_LOG(LogAndroidMedia, VeryVerbose, TEXT("Player %p: Disabling audio"), this, TrackIndex);
 
 				JavaMediaPlayer->SetAudioEnabled(false);
 			}
@@ -964,7 +1009,7 @@ bool FAndroidMediaPlayer::SelectTrack(EMediaTrackType TrackType, int32 TrackInde
 					return false;
 				}
 
-				UE_LOG(LogAndroidMedia, VeryVerbose, TEXT("Player %llx: Enabling audio"), this, TrackIndex);
+				UE_LOG(LogAndroidMedia, VeryVerbose, TEXT("Player %p: Enabling audio"), this, TrackIndex);
 
 				JavaMediaPlayer->SetAudioEnabled(true);
 			}
@@ -976,11 +1021,11 @@ bool FAndroidMediaPlayer::SelectTrack(EMediaTrackType TrackType, int32 TrackInde
 	case EMediaTrackType::Caption:
 		if (TrackIndex != SelectedCaptionTrack)
 		{
-			UE_LOG(LogAndroidMedia, Verbose, TEXT("Player %llx: Selecting caption track %i instead of %i (%i tracks)"), this, TrackIndex, SelectedCaptionTrack, CaptionTracks.Num());
+			UE_LOG(LogAndroidMedia, Verbose, TEXT("Player %p: Selecting caption track %i instead of %i (%i tracks)"), this, TrackIndex, SelectedCaptionTrack, CaptionTracks.Num());
 
 			if (TrackIndex == INDEX_NONE)
 			{
-				UE_LOG(LogAndroidMedia, VeryVerbose, TEXT("Player %llx: Disabling captions"), this, TrackIndex);
+				UE_LOG(LogAndroidMedia, VeryVerbose, TEXT("Player %p: Disabling captions"), this, TrackIndex);
 			}
 			else
 			{
@@ -989,7 +1034,7 @@ bool FAndroidMediaPlayer::SelectTrack(EMediaTrackType TrackType, int32 TrackInde
 					return false;
 				}
 
-				UE_LOG(LogAndroidMedia, VeryVerbose, TEXT("Player %llx: Enabling captions"), this, TrackIndex);
+				UE_LOG(LogAndroidMedia, VeryVerbose, TEXT("Player %p: Enabling captions"), this, TrackIndex);
 			}
 
 			SelectedCaptionTrack = TrackIndex;
@@ -999,11 +1044,11 @@ bool FAndroidMediaPlayer::SelectTrack(EMediaTrackType TrackType, int32 TrackInde
 	case EMediaTrackType::Video:
 		if (TrackIndex != SelectedVideoTrack)
 		{
-			UE_LOG(LogAndroidMedia, Verbose, TEXT("Player %llx: Selecting video track %i instead of %i (%i tracks)."), this, TrackIndex, SelectedVideoTrack, VideoTracks.Num());
+			UE_LOG(LogAndroidMedia, Verbose, TEXT("Player %p: Selecting video track %i instead of %i (%i tracks)."), this, TrackIndex, SelectedVideoTrack, VideoTracks.Num());
 
 			if (TrackIndex == INDEX_NONE)
 			{
-				UE_LOG(LogAndroidMedia, VeryVerbose, TEXT("Player %llx: Disabling video"), this, TrackIndex);
+				UE_LOG(LogAndroidMedia, VeryVerbose, TEXT("Player %p: Disabling video"), this, TrackIndex);
 				JavaMediaPlayer->SetVideoEnabled(false);
 			}
 			else
@@ -1013,7 +1058,7 @@ bool FAndroidMediaPlayer::SelectTrack(EMediaTrackType TrackType, int32 TrackInde
 					return false;
 				}
 
-				UE_LOG(LogAndroidMedia, VeryVerbose, TEXT("Player %llx: Enabling video"), this, TrackIndex);
+				UE_LOG(LogAndroidMedia, VeryVerbose, TEXT("Player %p: Enabling video"), this, TrackIndex);
 				JavaMediaPlayer->SetVideoEnabled(true);
 			}
 
@@ -1136,7 +1181,7 @@ bool FAndroidMediaPlayer::IsLooping() const
 
 bool FAndroidMediaPlayer::Seek(const FTimespan& Time)
 {
-	UE_LOG(LogAndroidMedia, Verbose, TEXT("Player %llx: Seeking to %s"), this, *Time.ToString());
+	UE_LOG(LogAndroidMedia, Verbose, TEXT("Player %p: Seeking to %s"), this, *Time.ToString());
 
 	if ((CurrentState == EMediaState::Closed) ||
 		(CurrentState == EMediaState::Error) ||
@@ -1146,7 +1191,7 @@ bool FAndroidMediaPlayer::Seek(const FTimespan& Time)
 		return false;
 	}
 
-	UE_LOG(LogAndroidMedia, Verbose, TEXT("Player %llx: Seeking to %s"), this, *Time.ToString());
+	UE_LOG(LogAndroidMedia, Verbose, TEXT("Player %p: Seeking to %s"), this, *Time.ToString());
 
 	JavaMediaPlayer->SeekTo(static_cast<int32>(Time.GetTotalMilliseconds()));
 	EventSink.ReceiveMediaEvent(EMediaEvent::SeekCompleted);
@@ -1185,7 +1230,7 @@ bool FAndroidMediaPlayer::SetRate(float Rate)
 		return true; // rate already set
 	}
 
-	UE_LOG(LogAndroidMedia, Verbose, TEXT("Player %llx: Setting rate from to %f to %f"), this, GetRate(), Rate);
+	UE_LOG(LogAndroidMedia, Verbose, TEXT("Player %p: Setting rate from to %f to %f"), this, GetRate(), Rate);
 
 	if (Rate == 0.0f)
 	{
diff --git a/Engine/Plugins/Media/AndroidMedia/Source/AndroidMedia/Private/Player/AndroidMediaPlayer.h b/Engine/Plugins/Media/AndroidMedia/Source/AndroidMedia/Private/Player/AndroidMediaPlayer.h
--- a/Engine/Plugins/Media/AndroidMedia/Source/AndroidMedia/Private/Player/AndroidMediaPlayer.h
+++ b/Engine/Plugins/Media/AndroidMedia/Source/AndroidMedia/Private/Player/AndroidMediaPlayer.h
@@ -51,7 +51,7 @@ public:
 	virtual IMediaCache& GetCache() override;
 	virtual IMediaControls& GetControls() override;
 	virtual FString GetInfo() const override;
-	virtual FName GetName() const override;
+	virtual FName GetPlayerName() const override;
 	virtual IMediaSamples& GetSamples() override;
 	virtual FString GetStats() const override;
 	virtual IMediaTracks& GetTracks() override;
@@ -60,8 +60,8 @@ public:
 	virtual bool Open(const FString& Url, const IMediaOptions* Options) override;
 	virtual bool Open(const TSharedRef<FArchive, ESPMode::ThreadSafe>& Archive, const FString& OriginalUrl, const IMediaOptions* Options) override;
 	virtual void SetGuid(const FGuid& Guid) override;
-	virtual void TickFetch(FTimespan DeltaTime) override;
-	virtual void TickInput(FTimespan DeltaTime) override;
+	virtual void TickFetch(FTimespan DeltaTime, FTimespan Timecode) override;
+	virtual void TickInput(FTimespan DeltaTime, FTimespan Timecode) override;
 
 protected:
 
diff --git a/Engine/Plugins/Media/AndroidMedia/Source/AndroidMediaFactory/Private/AndroidMediaFactoryModule.cpp b/Engine/Plugins/Media/AndroidMedia/Source/AndroidMediaFactory/Private/AndroidMediaFactoryModule.cpp
--- a/Engine/Plugins/Media/AndroidMedia/Source/AndroidMediaFactory/Private/AndroidMediaFactoryModule.cpp
+++ b/Engine/Plugins/Media/AndroidMedia/Source/AndroidMediaFactory/Private/AndroidMediaFactoryModule.cpp
@@ -88,7 +88,7 @@ public:
 		return LOCTEXT("MediaPlayerDisplayName", "Android Media");
 	}
 
-	virtual FName GetName() const override
+	virtual FName GetPlayerName() const override
 	{
 		static FName PlayerName(TEXT("AndroidMedia"));
 		return PlayerName;
diff --git a/Engine/Plugins/Media/AvfMedia/Source/AvfMedia/Private/Player/AvfMediaAudioSample.h b/Engine/Plugins/Media/AvfMedia/Source/AvfMedia/Private/Player/AvfMediaAudioSample.h
--- a/Engine/Plugins/Media/AvfMedia/Source/AvfMedia/Private/Player/AvfMediaAudioSample.h
+++ b/Engine/Plugins/Media/AvfMedia/Source/AvfMedia/Private/Player/AvfMediaAudioSample.h
@@ -67,7 +67,7 @@ public:
 		Buffer.AddUninitialized(InBufferSize);
 
 		Channels = InChannels;
-		Duration = FTimespan::FromSeconds(InFrames / InSampleRate);
+		Duration = FTimespan::FromSeconds((float)InFrames / (float)InSampleRate);
 		Frames = InFrames;
 		SampleRate = InSampleRate;
 		Time = InTime;
diff --git a/Engine/Plugins/Media/AvfMedia/Source/AvfMedia/Private/Player/AvfMediaPlayer.cpp b/Engine/Plugins/Media/AvfMedia/Source/AvfMedia/Private/Player/AvfMediaPlayer.cpp
--- a/Engine/Plugins/Media/AvfMedia/Source/AvfMedia/Private/Player/AvfMediaPlayer.cpp
+++ b/Engine/Plugins/Media/AvfMedia/Source/AvfMedia/Private/Player/AvfMediaPlayer.cpp
@@ -303,7 +303,7 @@ FString FAvfMediaPlayer::GetInfo() const
 }
 
 
-FName FAvfMediaPlayer::GetName() const
+FName FAvfMediaPlayer::GetPlayerName() const
 {
 	static FName PlayerName(TEXT("AvfMedia"));
 	return PlayerName;
@@ -475,7 +475,7 @@ void FAvfMediaPlayer::TickAudio()
 }
 
 
-void FAvfMediaPlayer::TickFetch(FTimespan DeltaTime)
+void FAvfMediaPlayer::TickFetch(FTimespan DeltaTime, FTimespan /*Timecode*/)
 {
 	if ((CurrentState > EMediaState::Error) && (Duration > 0.0f))
 	{
@@ -484,7 +484,7 @@ void FAvfMediaPlayer::TickFetch(FTimespan DeltaTime)
 }
 
 
-void FAvfMediaPlayer::TickInput(FTimespan DeltaTime)
+void FAvfMediaPlayer::TickInput(FTimespan DeltaTime, FTimespan /*Timecode*/)
 {
 	// process deferred tasks
 	TFunction<void()> Task;
diff --git a/Engine/Plugins/Media/AvfMedia/Source/AvfMedia/Private/Player/AvfMediaTracks.cpp b/Engine/Plugins/Media/AvfMedia/Source/AvfMedia/Private/Player/AvfMediaTracks.cpp
--- a/Engine/Plugins/Media/AvfMedia/Source/AvfMedia/Private/Player/AvfMediaTracks.cpp
+++ b/Engine/Plugins/Media/AvfMedia/Source/AvfMedia/Private/Player/AvfMediaTracks.cpp
@@ -375,7 +375,7 @@ void FAvfMediaTracks::ProcessAudio()
 
 				if (!AudioSample->Initialize(
 					OutputLength,
-					TargetDesc.mFramesPerPacket,
+					NumSamples,
 					TargetDesc.mChannelsPerFrame,
 					TargetDesc.mSampleRate,
 					FTimespan::FromSeconds(CMTimeGetSeconds(FrameTimeStamp))))
diff --git a/Engine/Plugins/Media/ImgMedia/Source/ImgMedia/Private/Player/ImgMediaPlayer.cpp b/Engine/Plugins/Media/ImgMedia/Source/ImgMedia/Private/Player/ImgMediaPlayer.cpp
--- a/Engine/Plugins/Media/ImgMedia/Source/ImgMedia/Private/Player/ImgMediaPlayer.cpp
+++ b/Engine/Plugins/Media/ImgMedia/Source/ImgMedia/Private/Player/ImgMediaPlayer.cpp
@@ -83,7 +83,7 @@ FString FImgMediaPlayer::GetInfo() const
 }
 
 
-FName FImgMediaPlayer::GetName() const
+FName FImgMediaPlayer::GetPlayerName() const
 {
 	static FName PlayerName(TEXT("ImgMedia"));
 	return PlayerName;
@@ -177,7 +177,7 @@ bool FImgMediaPlayer::Open(const TSharedRef<FArchive, ESPMode::ThreadSafe>& /*Ar
 }
 
 
-void FImgMediaPlayer::TickInput(FTimespan DeltaTime)
+void FImgMediaPlayer::TickInput(FTimespan DeltaTime, FTimespan /*Timecode*/)
 {
 	if (!Loader.IsValid() || (CurrentState == EMediaState::Error))
 	{
diff --git a/Engine/Plugins/Media/ImgMedia/Source/ImgMedia/Private/Player/ImgMediaPlayer.h b/Engine/Plugins/Media/ImgMedia/Source/ImgMedia/Private/Player/ImgMediaPlayer.h
--- a/Engine/Plugins/Media/ImgMedia/Source/ImgMedia/Private/Player/ImgMediaPlayer.h
+++ b/Engine/Plugins/Media/ImgMedia/Source/ImgMedia/Private/Player/ImgMediaPlayer.h
@@ -47,7 +47,7 @@ public:
 	virtual IMediaCache& GetCache() override;
 	virtual IMediaControls& GetControls() override;
 	virtual FString GetInfo() const override;
-	virtual FName GetName() const override;
+	virtual FName GetPlayerName() const override;
 	virtual IMediaSamples& GetSamples() override;
 	virtual FString GetStats() const override;
 	virtual IMediaTracks& GetTracks() override;
@@ -55,7 +55,7 @@ public:
 	virtual IMediaView& GetView() override;
 	virtual bool Open(const FString& Url, const IMediaOptions* Options) override;
 	virtual bool Open(const TSharedRef<FArchive, ESPMode::ThreadSafe>& Archive, const FString& OriginalUrl, const IMediaOptions* Options) override;
-	virtual void TickInput(FTimespan DeltaTime) override;
+	virtual void TickInput(FTimespan DeltaTime, FTimespan Timecode) override;
 
 protected:
 
diff --git a/Engine/Plugins/Media/MediaPlayerEditor/Source/MediaPlayerEditor/Private/Customizations/BaseMediaSourceCustomization.cpp b/Engine/Plugins/Media/MediaPlayerEditor/Source/MediaPlayerEditor/Private/Customizations/BaseMediaSourceCustomization.cpp
--- a/Engine/Plugins/Media/MediaPlayerEditor/Source/MediaPlayerEditor/Private/Customizations/BaseMediaSourceCustomization.cpp
+++ b/Engine/Plugins/Media/MediaPlayerEditor/Source/MediaPlayerEditor/Private/Customizations/BaseMediaSourceCustomization.cpp
@@ -93,7 +93,7 @@ TSharedRef<SWidget> FBaseMediaSourceCustomization::MakePlatformPlayersMenu(const
 		for (IMediaPlayerFactory* Factory : PlayerFactories)
 		{
 			const bool SupportsPlatform = Factory->GetSupportedPlatforms().Contains(*IniPlatformName);
-			const FName PlayerName = Factory->GetName();
+			const FName PlayerName = Factory->GetPlayerName();
 
 			MenuBuilder.AddMenuEntry(
 				FText::Format(LOCTEXT("PlayerNameFormat", "{0} ({1})"), Factory->GetDisplayName(), FText::FromName(PlayerName)),
diff --git a/Engine/Plugins/Media/MediaPlayerEditor/Source/MediaPlayerEditor/Private/Toolkits/MediaPlayerEditorToolkit.cpp b/Engine/Plugins/Media/MediaPlayerEditor/Source/MediaPlayerEditor/Private/Toolkits/MediaPlayerEditorToolkit.cpp
--- a/Engine/Plugins/Media/MediaPlayerEditor/Source/MediaPlayerEditor/Private/Toolkits/MediaPlayerEditorToolkit.cpp
+++ b/Engine/Plugins/Media/MediaPlayerEditor/Source/MediaPlayerEditor/Private/Toolkits/MediaPlayerEditorToolkit.cpp
@@ -306,7 +306,7 @@ void FMediaPlayerEditorToolkit::BindCommands()
 	ToolkitCommands->MapAction(
 		Commands.PlayMedia,
 		FExecuteAction::CreateLambda([this]{ MediaPlayer->Play(); }),
-		FCanExecuteAction::CreateLambda([this]{ return MediaPlayer->IsReady() && (MediaPlayer->GetRate() != 1.0f); })
+		FCanExecuteAction::CreateLambda([this]{ return MediaPlayer->IsReady() && (!MediaPlayer->IsPlaying() || (MediaPlayer->GetRate() != 1.0f)); })
 	);
 
 	ToolkitCommands->MapAction(
diff --git a/Engine/Plugins/Media/MediaPlayerEditor/Source/MediaPlayerEditor/Private/Widgets/SMediaPlayerEditorViewer.cpp b/Engine/Plugins/Media/MediaPlayerEditor/Source/MediaPlayerEditor/Private/Widgets/SMediaPlayerEditorViewer.cpp
--- a/Engine/Plugins/Media/MediaPlayerEditor/Source/MediaPlayerEditor/Private/Widgets/SMediaPlayerEditorViewer.cpp
+++ b/Engine/Plugins/Media/MediaPlayerEditor/Source/MediaPlayerEditor/Private/Widgets/SMediaPlayerEditorViewer.cpp
@@ -732,7 +732,7 @@ void SMediaPlayerEditorViewer::HandleDecoderMenuNewMenu(FMenuBuilder& MenuBuilde
 	for (IMediaPlayerFactory* Factory : PlayerFactories)
 	{
 		const bool SupportsRunningPlatform = Factory->GetSupportedPlatforms().Contains(PlatformName);
-		const FName PlayerName = Factory->GetName();
+		const FName PlayerName = Factory->GetPlayerName();
 
 		MenuBuilder.AddMenuEntry(
 			FText::Format(LOCTEXT("PlayerNameFormat", "{0} ({1})"), Factory->GetDisplayName(), FText::FromName(PlayerName)),
diff --git a/Engine/Plugins/Media/MediaPlayerEditor/Source/MediaPlayerEditor/Private/Widgets/SMediaPlayerEditorViewport.cpp b/Engine/Plugins/Media/MediaPlayerEditor/Source/MediaPlayerEditor/Private/Widgets/SMediaPlayerEditorViewport.cpp
--- a/Engine/Plugins/Media/MediaPlayerEditor/Source/MediaPlayerEditor/Private/Widgets/SMediaPlayerEditorViewport.cpp
+++ b/Engine/Plugins/Media/MediaPlayerEditor/Source/MediaPlayerEditor/Private/Widgets/SMediaPlayerEditorViewport.cpp
@@ -286,12 +286,14 @@ FText SMediaPlayerEditorViewport::HandleMediaPlayerStateText() const
 
 FText SMediaPlayerEditorViewport::HandleMediaSourceNameText() const
 {
-	if (MediaPlayer->GetUrl().IsEmpty())
+	const FText MediaName = MediaPlayer->GetMediaName();
+
+	if (MediaName.IsEmpty())
 	{
 		return LOCTEXT("StateOverlayNoMedia", "No Media");
 	}
 
-	return FText::FromString(FPaths::GetBaseFilename(MediaPlayer->GetUrl()));
+	return MediaName;
 }
 
 
diff --git a/Engine/Plugins/Media/MfMedia/Source/MfMedia/Private/Mf/MfMediaSourceReaderCallback.cpp b/Engine/Plugins/Media/MfMedia/Source/MfMedia/Private/Mf/MfMediaSourceReaderCallback.cpp
--- a/Engine/Plugins/Media/MfMedia/Source/MfMedia/Private/Mf/MfMediaSourceReaderCallback.cpp
+++ b/Engine/Plugins/Media/MfMedia/Source/MfMedia/Private/Mf/MfMediaSourceReaderCallback.cpp
@@ -31,7 +31,7 @@ FMfMediaSourceReaderCallback::FMfMediaSourceReaderCallback(IMfMediaSourceReaderS
 
 FMfMediaSourceReaderCallback::~FMfMediaSourceReaderCallback()
 {
-	UE_LOG(LogMfMedia, VeryVerbose, TEXT("Callback %llx: Destroyed)"), this);
+	UE_LOG(LogMfMedia, VeryVerbose, TEXT("Callback %p: Destroyed)"), this);
 }
 
 
@@ -47,7 +47,7 @@ STDMETHODIMP FMfMediaSourceReaderCallback::OnEvent(DWORD, IMFMediaEvent* Event)
 
 		if (FAILED(Result))
 		{
-			UE_LOG(LogMfMedia, VeryVerbose, TEXT("Callback %llx: Failed to get session event type: %s"), this, *MfMedia::ResultToString(Result));
+			UE_LOG(LogMfMedia, VeryVerbose, TEXT("Callback %p: Failed to get session event type: %s"), this, *MfMedia::ResultToString(Result));
 			return Result;
 		}
 	}
@@ -59,12 +59,12 @@ STDMETHODIMP FMfMediaSourceReaderCallback::OnEvent(DWORD, IMFMediaEvent* Event)
 
 		if (FAILED(Result))
 		{
-			UE_LOG(LogMfMedia, VeryVerbose, TEXT("Callback %llx: Failed to get event status: %s"), this, *MfMedia::ResultToString(Result));
+			UE_LOG(LogMfMedia, VeryVerbose, TEXT("Callback %p: Failed to get event status: %s"), this, *MfMedia::ResultToString(Result));
 			return Result;
 		}
 	}
 
-	UE_LOG(LogMfMedia, VeryVerbose, TEXT("Callback %llx: Event [%s] (%s)"), this, *MfMedia::MediaEventToString(EventType), *MfMedia::ResultToString(EventStatus));
+	UE_LOG(LogMfMedia, VeryVerbose, TEXT("Callback %p: Event [%s]: %s"), this, *MfMedia::MediaEventToString(EventType), *MfMedia::ResultToString(EventStatus));
 
 	Sink.ReceiveSourceReaderEvent(EventType);
 
@@ -83,7 +83,7 @@ STDMETHODIMP FMfMediaSourceReaderCallback::OnFlush(DWORD)
 STDMETHODIMP FMfMediaSourceReaderCallback::OnReadSample(HRESULT Status, DWORD StreamIndex, DWORD StreamFlags, LONGLONG Timestamp, IMFSample* Sample)
 {
 	#if MFMEDIASOURCEREADERCALLBACK_TRACE_SAMPLES
-		UE_LOG(LogMfMedia, VeryVerbose, TEXT("Callback %llx: Sample read: %s (stream = %i, status = %i, flags = %i)"), this, *FTimespan(Timestamp).ToString(), StreamIndex, Status, StreamFlags);
+		UE_LOG(LogMfMedia, VeryVerbose, TEXT("Callback %p: Sample read: %s (stream = %i, status = %i, flags = %i)"), this, *FTimespan(Timestamp).ToString(), StreamIndex, Status, StreamFlags);
 	#endif
 
 	Sink.ReceiveSourceReaderSample(Sample, Status, StreamFlags, StreamIndex, FTimespan(Timestamp));
diff --git a/Engine/Plugins/Media/MfMedia/Source/MfMedia/Private/Mf/MfMediaUtils.cpp b/Engine/Plugins/Media/MfMedia/Source/MfMedia/Private/Mf/MfMediaUtils.cpp
--- a/Engine/Plugins/Media/MfMedia/Source/MfMedia/Private/Mf/MfMediaUtils.cpp
+++ b/Engine/Plugins/Media/MfMedia/Source/MfMedia/Private/Mf/MfMediaUtils.cpp
@@ -30,7 +30,7 @@ namespace MfMedia
 
 			if (FAILED(Result))
 			{
-				UE_LOG(LogMfMedia, Warning, TEXT("Failed to create %s output type (%s)"), *MajorTypeToString(MajorType), *ResultToString(Result));
+				UE_LOG(LogMfMedia, Warning, TEXT("Failed to create %s output type: %s"), *MajorTypeToString(MajorType), *ResultToString(Result));
 				return NULL;
 			}
 
@@ -38,7 +38,7 @@ namespace MfMedia
 
 			if (FAILED(Result))
 			{
-				UE_LOG(LogMfMedia, Warning, TEXT("Failed to initialize %s output type (%s)"), *MajorTypeToString(MajorType), *ResultToString(Result));
+				UE_LOG(LogMfMedia, Warning, TEXT("Failed to initialize %s output type: %s"), *MajorTypeToString(MajorType), *ResultToString(Result));
 				return NULL;
 			}
 		}
@@ -115,7 +115,7 @@ namespace MfMedia
 
 			if (FAILED(Result))
 			{
-				UE_LOG(LogMfMedia, Warning, TEXT("Failed to initialize caption output type (%s)"), *ResultToString(Result));
+				UE_LOG(LogMfMedia, Warning, TEXT("Failed to initialize caption output type: %s"), *ResultToString(Result));
 				return NULL;
 			}
 		}
@@ -162,11 +162,11 @@ namespace MfMedia
 
 			if (FAILED(Result))
 			{
-				UE_LOG(LogMfMedia, Warning, TEXT("Failed to set video output type (%s)"), *ResultToString(Result));
+				UE_LOG(LogMfMedia, Warning, TEXT("Failed to set video output type: %s"), *ResultToString(Result));
 				return NULL;
 			}
 
-#if 1//PLATFORM_XBOXONE
+#if PLATFORM_XBOXONE
 			Result = OutputType->SetGUID(MF_MT_SUBTYPE, MFVideoFormat_NV12); // XboxOne only supports NV12
 #else
 			Result = OutputType->SetGUID(MF_MT_SUBTYPE, (SubType == MFVideoFormat_RGB32) ? MFVideoFormat_RGB32 : MFVideoFormat_NV12);
@@ -174,7 +174,7 @@ namespace MfMedia
 
 			if (FAILED(Result))
 			{
-				UE_LOG(LogMfMedia, Warning, TEXT("Failed to set video output sub-type (%s)"), *ResultToString(Result));
+				UE_LOG(LogMfMedia, Warning, TEXT("Failed to set video output sub-type: %s"), *ResultToString(Result));
 				return NULL;
 			}
 		}
@@ -391,7 +391,7 @@ namespace MfMedia
 
 			if (FAILED(Result))
 			{
-				UE_LOG(LogMfMedia, Error, TEXT("Failed to create media source resolver (%s)"), *MfMedia::ResultToString(Result));
+				UE_LOG(LogMfMedia, Error, TEXT("Failed to create media source resolver: %s"), *MfMedia::ResultToString(Result));
 				return NULL;
 			}
 		}
@@ -408,7 +408,7 @@ namespace MfMedia
 
 				if (FAILED(Result))
 				{
-					UE_LOG(LogMfMedia, Error, TEXT("Failed to resolve byte stream %s (%s)"), *Url, *MfMedia::ResultToString(Result));
+					UE_LOG(LogMfMedia, Error, TEXT("Failed to resolve byte stream %s: %s"), *Url, *MfMedia::ResultToString(Result));
 					return NULL;
 				}
 			}
@@ -418,7 +418,7 @@ namespace MfMedia
 
 				if (FAILED(Result))
 				{
-					UE_LOG(LogMfMedia, Error, TEXT("Failed to resolve URL %s (%s)"), *Url, *MfMedia::ResultToString(Result));
+					UE_LOG(LogMfMedia, Error, TEXT("Failed to resolve URL %s: %s"), *Url, *MfMedia::ResultToString(Result));
 					return NULL;
 				}
 			}
diff --git a/Engine/Plugins/Media/MfMedia/Source/MfMedia/Private/Player/MfMediaPlayer.cpp b/Engine/Plugins/Media/MfMedia/Source/MfMedia/Private/Player/MfMediaPlayer.cpp
--- a/Engine/Plugins/Media/MfMedia/Source/MfMedia/Private/Player/MfMediaPlayer.cpp
+++ b/Engine/Plugins/Media/MfMedia/Source/MfMedia/Private/Player/MfMediaPlayer.cpp
@@ -112,7 +112,7 @@ FString FMfMediaPlayer::GetInfo() const
 }
 
 
-FName FMfMediaPlayer::GetName() const
+FName FMfMediaPlayer::GetPlayerName() const
 {
 	static FName PlayerName(TEXT("MfMedia"));
 	return PlayerName;
@@ -174,13 +174,13 @@ bool FMfMediaPlayer::Open(const TSharedRef<FArchive, ESPMode::ThreadSafe>& Archi
 
 	if (Archive->TotalSize() == 0)
 	{
-		UE_LOG(LogMfMedia, Error, TEXT("Cannot open media from archive (archive is empty)."));
+		UE_LOG(LogMfMedia, Verbose, TEXT("Player %p: Cannot open media from archive (archive is empty)."), this);
 		return false;
 	}
 
 	if (OriginalUrl.IsEmpty())
 	{
-		UE_LOG(LogMfMedia, Error, TEXT("Cannot open media from archive (no original URL provided)."));
+		UE_LOG(LogMfMedia, Verbose, TEXT("Player %p: Cannot open media from archive (no original URL provided)."), this);
 		return false;
 	}
 
@@ -197,7 +197,7 @@ void FMfMediaPlayer::TickAudio()
 }
 
 
-void FMfMediaPlayer::TickFetch(FTimespan /*DeltaTime*/)
+void FMfMediaPlayer::TickFetch(FTimespan /*DeltaTime*/, FTimespan /*Timecode*/)
 {
 	bool MediaSourceChanged = false;
 	bool TrackSelectionChanged = false;
@@ -252,7 +252,7 @@ void FMfMediaPlayer::TickFetch(FTimespan /*DeltaTime*/)
 }
 
 
-void FMfMediaPlayer::TickInput(FTimespan DeltaTime)
+void FMfMediaPlayer::TickInput(FTimespan DeltaTime, FTimespan Timecode)
 {
 	if ((CurrentState != EMediaState::Playing) || (CurrentDuration == FTimespan::Zero()))
 	{
@@ -322,7 +322,7 @@ void FMfMediaPlayer::TickInput(FTimespan DeltaTime)
 	// update tracks
 	if (CurrentState == EMediaState::Playing)
 	{
-		Tracks->TickInput(CurrentRate, CurrentTime, DeltaTime);
+		Tracks->TickInput(CurrentRate, CurrentTime);
 	}
 }
 
@@ -349,7 +349,7 @@ bool FMfMediaPlayer::CommitTime(FTimespan Time)
 
 	if (FAILED(Result))
 	{
-		UE_LOG(LogMfMedia, Warning, TEXT("Failed to commit time %s: %s"), *Time.ToString(), *MfMedia::ResultToString(Result));
+		UE_LOG(LogMfMedia, Verbose, TEXT("Player %p: Failed to set source reader position to %s: %s"), this, *Time.ToString(), *MfMedia::ResultToString(Result));
 		return false;
 	}
 
@@ -361,7 +361,7 @@ bool FMfMediaPlayer::CommitTime(FTimespan Time)
 
 bool FMfMediaPlayer::InitializePlayer(const TSharedPtr<FArchive, ESPMode::ThreadSafe>& Archive, const FString& Url, bool Precache)
 {
-	UE_LOG(LogMfMedia, VeryVerbose, TEXT("Player %llx: Initializing %s (archive = %s, precache = %s)"), this, *Url, Archive.IsValid() ? TEXT("yes") : TEXT("no"), Precache ? TEXT("yes") : TEXT("no"));
+	UE_LOG(LogMfMedia, VeryVerbose, TEXT("Player %p: Initializing %s (archive = %s, precache = %s)"), this, *Url, Archive.IsValid() ? TEXT("yes") : TEXT("no"), Precache ? TEXT("yes") : TEXT("no"));
 
 	CurrentState = EMediaState::Preparing;
 
@@ -393,6 +393,15 @@ bool FMfMediaPlayer::InitializePlayer(const TSharedPtr<FArchive, ESPMode::Thread
 
 void FMfMediaPlayer::UpdateCharacteristics()
 {
+	// reset characteristics
+	Characteristics = 0u;
+
+	RateControl.Reset();
+	RateSupport.Reset();
+
+	ThinnedRates.Empty();
+	UnthinnedRates.Empty();
+
 	if (!MediaSource.IsValid())
 	{
 		return;
@@ -403,7 +412,7 @@ void FMfMediaPlayer::UpdateCharacteristics()
 
 	if (FAILED(Result))
 	{
-		UE_LOG(LogMfMedia, Warning, TEXT("Failed to get media source characteristics: %s"), *MfMedia::ResultToString(Result));
+		UE_LOG(LogMfMedia, Verbose, TEXT("Player %p: Failed to get media source characteristics: %s"), this, *MfMedia::ResultToString(Result));
 	}
 
 	// get service interface
@@ -413,7 +422,7 @@ void FMfMediaPlayer::UpdateCharacteristics()
 
 	if (FAILED(Result))
 	{
-		UE_LOG(LogMfMedia, Warning, TEXT("Failed to query service interface: %s"), *MfMedia::ResultToString(Result));
+		UE_LOG(LogMfMedia, Verbose, TEXT("Player %p: Failed to query service interface: %s"), this, *MfMedia::ResultToString(Result));
 		return;
 	}
 
@@ -422,22 +431,22 @@ void FMfMediaPlayer::UpdateCharacteristics()
 
 	if (FAILED(Result))
 	{
-		UE_LOG(LogMfMedia, Log, TEXT("Rate control service unavailable (%s)"), *MfMedia::ResultToString(Result));
+		UE_LOG(LogMfMedia, Log, TEXT("Rate control service unavailable: %s"), *MfMedia::ResultToString(Result));
 	}
 	else
 	{
-		UE_LOG(LogMfMedia, Verbose, TEXT("Player %llx: Rate control ready"), this);
+		UE_LOG(LogMfMedia, Verbose, TEXT("Player %p: Rate control ready"), this);
 	}
 
 	Result = GetService->GetService(MF_RATE_CONTROL_SERVICE, IID_PPV_ARGS(&RateSupport));
 
 	if (FAILED(Result))
 	{
-		UE_LOG(LogMfMedia, Log, TEXT("Rate support service unavailable (%s)"), *MfMedia::ResultToString(Result));
+		UE_LOG(LogMfMedia, Log, TEXT("Rate support service unavailable: %s"), *MfMedia::ResultToString(Result));
 	}
 	else
 	{
-		UE_LOG(LogMfMedia, Verbose, TEXT("Player %llx: Rate support ready"), this);
+		UE_LOG(LogMfMedia, Verbose, TEXT("Player %p: Rate support ready"), this);
 	}
 
 	// cache rate control properties
@@ -474,11 +483,6 @@ void FMfMediaPlayer::UpdateCharacteristics()
 		}
 #endif
 	}
-	else
-	{
-		ThinnedRates.Empty();
-		UnthinnedRates.Empty();
-	}
 }
 
 
@@ -570,17 +574,17 @@ bool FMfMediaPlayer::Seek(const FTimespan& Time)
 	// validate seek
 	if ((CurrentState == EMediaState::Closed) || (CurrentState == EMediaState::Error))
 	{
-		UE_LOG(LogMfMedia, Warning, TEXT("Cannot seek while closed or in error state"));
+		UE_LOG(LogMfMedia, Verbose, TEXT("Player %p: Cannot seek while closed or in error state"), this);
 		return false;
 	}
 
 	if ((Time < FTimespan::Zero()) || (Time > CurrentDuration))
 	{
-		UE_LOG(LogMfMedia, Warning, TEXT("Invalid seek time %s (media duration is %s)"), *Time.ToString(), *CurrentDuration.ToString());
+		UE_LOG(LogMfMedia, Verbose, TEXT("Player %p: Invalid seek time %s (media duration is %s)"), this, *Time.ToString(), *CurrentDuration.ToString());
 		return false;
 	}
 	
-	UE_LOG(LogMfMedia, Verbose, TEXT("Player %llx: Seeking to %s"), this, *Time.ToString());
+	UE_LOG(LogMfMedia, Verbose, TEXT("Player %p: Seeking to %s"), this, *Time.ToString());
 
 	if (!CommitTime(Time))
 	{
@@ -623,17 +627,17 @@ bool FMfMediaPlayer::SetRate(float Rate)
 	// check whether rate is supported
 	if (UnthinnedRates.Contains(Rate))
 	{
-		UE_LOG(LogMfMedia, Verbose, TEXT("Player %llx: Setting rate from to %f to %f (unthinned)"), this, CurrentRate, Rate);
+		UE_LOG(LogMfMedia, Verbose, TEXT("Player %p: Setting rate from to %f to %f (unthinned)"), this, CurrentRate, Rate);
 		Thin = FALSE;
 	}
 	else if (ThinnedRates.Contains(Rate))
 	{
-		UE_LOG(LogMfMedia, Verbose, TEXT("Player %llx: Setting rate from to %f to %f (thinned)"), this, CurrentRate, Rate);
+		UE_LOG(LogMfMedia, Verbose, TEXT("Player %p: Setting rate from to %f to %f (thinned)"), this, CurrentRate, Rate);
 		Thin = TRUE;
 	}
 	else
 	{
-		UE_LOG(LogMfMedia, Warning, TEXT("The rate %f is not supported"), Rate);
+		UE_LOG(LogMfMedia, Verbose, TEXT("Player %p: The rate %f is not supported"), this, Rate);
 		return false;
 	}
 
@@ -649,7 +653,7 @@ bool FMfMediaPlayer::SetRate(float Rate)
 
 			if (FAILED(Result))
 			{
-				UE_LOG(LogMfMedia, Warning, TEXT("Failed to commit rate change from %f to zero: %s"), CurrentRate, *MfMedia::ResultToString(Result));
+				UE_LOG(LogMfMedia, Verbose, TEXT("Player %p: Failed to commit rate change from %f to zero: %s"), this, CurrentRate, *MfMedia::ResultToString(Result));
 				return false;
 			}
 		}
@@ -663,7 +667,7 @@ bool FMfMediaPlayer::SetRate(float Rate)
 
 			if (FAILED(Result))
 			{
-				UE_LOG(LogMfMedia, Warning, TEXT("Failed to commit rate change from %f to %f (%s): %s"), CurrentRate, Rate, ThinnedString, *MfMedia::ResultToString(Result));
+				UE_LOG(LogMfMedia, Verbose, TEXT("Player %p: Failed to commit rate change from %f to %f [%s]: %s"), this, CurrentRate, Rate, ThinnedString, *MfMedia::ResultToString(Result));
 				return false;
 			}
 		}
diff --git a/Engine/Plugins/Media/MfMedia/Source/MfMedia/Private/Player/MfMediaTracks.cpp b/Engine/Plugins/Media/MfMedia/Source/MfMedia/Private/Player/MfMediaTracks.cpp
--- a/Engine/Plugins/Media/MfMedia/Source/MfMedia/Private/Player/MfMediaTracks.cpp
+++ b/Engine/Plugins/Media/MfMedia/Source/MfMedia/Private/Player/MfMediaTracks.cpp
@@ -8,6 +8,7 @@
 #include "IMediaTextureSample.h"
 #include "Internationalization/Internationalization.h"
 #include "Math/IntPoint.h"
+#include "MediaHelpers.h"
 #include "MediaSamples.h"
 #include "Misc/ScopeLock.h"
 
@@ -76,14 +77,14 @@ void FMfMediaTracks::AppendStats(FString &OutStats) const
 
 	if (AudioTracks.Num() == 0)
 	{
-		OutStats += TEXT("    none\n");
+		OutStats += TEXT("\tnone\n");
 	}
 	else
 	{
 		for (const FTrack& Track : AudioTracks)
 		{
-			OutStats += FString::Printf(TEXT("    %s\n"), *Track.DisplayName.ToString());
-			OutStats += TEXT("        Not implemented yet");
+			OutStats += FString::Printf(TEXT("\t%s\n"), *Track.DisplayName.ToString());
+			OutStats += TEXT("\t\tNot implemented yet");
 		}
 	}
 
@@ -92,14 +93,14 @@ void FMfMediaTracks::AppendStats(FString &OutStats) const
 
 	if (VideoTracks.Num() == 0)
 	{
-		OutStats += TEXT("    none\n");
+		OutStats += TEXT("\tnone\n");
 	}
 	else
 	{
 		for (const FTrack& Track : VideoTracks)
 		{
-			OutStats += FString::Printf(TEXT("    %s\n"), *Track.DisplayName.ToString());
-			OutStats += TEXT("        Not implemented yet");
+			OutStats += FString::Printf(TEXT("\t%s\n"), *Track.DisplayName.ToString());
+			OutStats += TEXT("\t\tNot implemented yet");
 		}
 	}
 }
@@ -157,6 +158,8 @@ void FMfMediaTracks::Initialize(IMFMediaSource* InMediaSource, IMFSourceReaderCa
 {
 	Shutdown();
 
+	UE_LOG(LogMfMedia, Verbose, TEXT("Tracks: %p: Initializing tracks"), this);
+
 	FScopeLock Lock(&CriticalSection);
 
 	MediaSourceChanged = true;
@@ -173,7 +176,7 @@ void FMfMediaTracks::Initialize(IMFMediaSource* InMediaSource, IMFSourceReaderCa
 
 		if (FAILED(Result))
 		{
-			UE_LOG(LogMfMedia, Error, TEXT("Failed to create presentation descriptor: %s"), *MfMedia::ResultToString(Result));
+			UE_LOG(LogMfMedia, Verbose, TEXT("Tracks %p: Failed to create presentation descriptor: %s"), this, *MfMedia::ResultToString(Result));
 			return;
 		}
 	}
@@ -183,27 +186,27 @@ void FMfMediaTracks::Initialize(IMFMediaSource* InMediaSource, IMFSourceReaderCa
 	{
 		if (FAILED(MFCreateAttributes(&Attributes, 1)))
 		{
-			UE_LOG(LogMfMedia, Error, TEXT("Failed to create source reader attributes"));
+			UE_LOG(LogMfMedia, Verbose, TEXT("Tracks %p: Failed to create source reader attributes"), this);
 			return;
 		}
 
 		if (FAILED(Attributes->SetUINT32(MF_READWRITE_ENABLE_HARDWARE_TRANSFORMS, TRUE)))
 		{
-			UE_LOG(LogMfMedia, Warning, TEXT("Failed to set one or more source reader attributes"));
+			UE_LOG(LogMfMedia, Verbose, TEXT("Tracks %p: Failed to set one or more source reader attributes"), this);
 		}
 
 #if MFMEDIATRACKS_USE_DXVA
 		if (FAILED(Attributes->SetUINT32(MF_SOURCE_READER_D3D_MANAGER, D3dDeviceManager)) ||
 			FAILED(Attributes->SetUINT32(MF_SOURCE_READER_DISABLE_DXVA, FALSE)))
 		{
-			UE_LOG(LogMfMedia, Warning, TEXT("Failed to set DXVA source reader attributes"));
+			UE_LOG(LogMfMedia, Verbose, TEXT("Tracks %p: Failed to set DXVA source reader attributes"), this);
 		}
 #endif
 
 #if MFMEDIATRACKS_USE_ASYNCREADER
 		if (FAILED(Attributes->SetUnknown(MF_SOURCE_READER_ASYNC_CALLBACK, InSourceReaderCallback)))
 		{
-			UE_LOG(LogMfMedia, Warning, TEXT("Failed to set async callback source reader attribute"));
+			UE_LOG(LogMfMedia, Verbose, TEXT("Tracks %p: Failed to set async callback source reader attribute"), this);
 		}
 #endif
 	}
@@ -214,7 +217,7 @@ void FMfMediaTracks::Initialize(IMFMediaSource* InMediaSource, IMFSourceReaderCa
 
 		if (FAILED(Result))
 		{
-			UE_LOG(LogMfMedia, Error, TEXT("Failed to create source reader: %s"), *MfMedia::ResultToString(Result));
+			UE_LOG(LogMfMedia, Verbose, TEXT("Tracks %p: Failed to create source reader: %s"), this, *MfMedia::ResultToString(Result));
 			return;
 		}
 	}
@@ -226,11 +229,11 @@ void FMfMediaTracks::Initialize(IMFMediaSource* InMediaSource, IMFSourceReaderCa
 
 		if (FAILED(Result))
 		{
-			UE_LOG(LogMfMedia, Error, TEXT("Failed to get stream count (%s)"), *MfMedia::ResultToString(Result));
+			UE_LOG(LogMfMedia, Verbose, TEXT("Tracks %p: Failed to get stream count: %s"), this, *MfMedia::ResultToString(Result));
 			return;
 		}
 
-		UE_LOG(LogMfMedia, Verbose, TEXT("Tracks %llx: Found %i streams"), this, StreamCount);
+		UE_LOG(LogMfMedia, Verbose, TEXT("Tracks %p: Found %i streams"), this, StreamCount);
 	}
 
 	// initialization successful
@@ -240,11 +243,18 @@ void FMfMediaTracks::Initialize(IMFMediaSource* InMediaSource, IMFSourceReaderCa
 	SourceReader = NewSourceReader;
 
 	// add streams (Media Foundation reports them in reverse order)
+	bool AllStreamsAdded = true;
+
 	for (int32 StreamIndex = StreamCount - 1; StreamIndex >= 0; --StreamIndex)
 	{
-		AddStreamToTracks(StreamIndex, Info);
+		AllStreamsAdded &= AddStreamToTracks(StreamIndex, Info);
 		Info += TEXT("\n");
 	}
+
+	if (!AllStreamsAdded)
+	{
+		UE_LOG(LogMfMedia, Verbose, TEXT("Tracks %p: Not all available streams were added to the track collection"), this);
+	}
 }
 
 
@@ -271,7 +281,7 @@ void FMfMediaTracks::ProcessSample(IMFSample* Sample, HRESULT Status, DWORD Stre
 		{
 			if (((StreamFlags & MF_SOURCE_READERF_ENDOFSTREAM) != 0) || ((StreamFlags & MF_SOURCE_READERF_ERROR) != 0))
 			{
-				UE_LOG(LogMfMedia, Verbose, TEXT("Tracks %llx: Audio done"), this);
+				UE_LOG(LogMfMedia, Verbose, TEXT("Tracks %p: Audio done"), this);
 
 				AudioDone = true;
 			}
@@ -289,7 +299,7 @@ void FMfMediaTracks::ProcessSample(IMFSample* Sample, HRESULT Status, DWORD Stre
 					LastAudioSampleTime = AudioSample->GetTime();
 
 					#if MFMEDIATRACKS_TRACE_SAMPLES
-						UE_LOG(LogMfMedia, VeryVerbose, TEXT("Tracks %llx: Audio sample processed: %s"), this, *LastAudioSampleTime.ToString());
+						UE_LOG(LogMfMedia, VeryVerbose, TEXT("Tracks %p: Audio sample processed: %s"), this, *LastAudioSampleTime.ToString());
 					#endif
 				}
 			}
@@ -310,7 +320,7 @@ void FMfMediaTracks::ProcessSample(IMFSample* Sample, HRESULT Status, DWORD Stre
 		{
 			if (((StreamFlags & MF_SOURCE_READERF_ENDOFSTREAM) != 0) || ((StreamFlags & MF_SOURCE_READERF_ERROR) != 0))
 			{
-				UE_LOG(LogMfMedia, Verbose, TEXT("Tracks %llx: Caption done"), this);
+				UE_LOG(LogMfMedia, Verbose, TEXT("Tracks %p: Caption done"), this);
 
 				CaptionDone = true;
 			}
@@ -332,7 +342,7 @@ void FMfMediaTracks::ProcessSample(IMFSample* Sample, HRESULT Status, DWORD Stre
 					LastCaptionSampleTime = FTimespan(SampleTime);
 
 					#if MFMEDIATRACKS_TRACE_SAMPLES
-						UE_LOG(LogMfMedia, VeryVerbose, TEXT("Tracks %llx: Caption sample processed: %s"), this, *LastCaptionSampleTime.ToString());
+						UE_LOG(LogMfMedia, VeryVerbose, TEXT("Tracks %p: Caption sample processed: %s"), this, *LastCaptionSampleTime.ToString());
 					#endif
 //				}
 			}
@@ -351,7 +361,7 @@ void FMfMediaTracks::ProcessSample(IMFSample* Sample, HRESULT Status, DWORD Stre
 		{
 			if (((StreamFlags & MF_SOURCE_READERF_ENDOFSTREAM) != 0) || ((StreamFlags & MF_SOURCE_READERF_ERROR) != 0))
 			{
-				UE_LOG(LogMfMedia, Verbose, TEXT("Tracks %llx: Video done"), this);
+				UE_LOG(LogMfMedia, Verbose, TEXT("Tracks %p: Video done"), this);
 
 				VideoDone = true;
 			}
@@ -369,7 +379,7 @@ void FMfMediaTracks::ProcessSample(IMFSample* Sample, HRESULT Status, DWORD Stre
 					LastVideoSampleTime = VideoSample->GetTime();
 
 					#if MFMEDIATRACKS_TRACE_SAMPLES
-						UE_LOG(LogMfMedia, VeryVerbose, TEXT("Tracks %llx: Video sample processed: %s"), this, *LastVideoSampleTime.ToString());
+						UE_LOG(LogMfMedia, VeryVerbose, TEXT("Tracks %p: Video sample processed: %s"), this, *LastVideoSampleTime.ToString());
 					#endif
 				}
 			}
@@ -385,7 +395,7 @@ void FMfMediaTracks::ProcessSample(IMFSample* Sample, HRESULT Status, DWORD Stre
 
 void FMfMediaTracks::Restart()
 {
-	UE_LOG(LogMfMedia, Verbose, TEXT("Tracks %llx: Restarting sample processing"), this);
+	UE_LOG(LogMfMedia, Verbose, TEXT("Tracks %p: Restarting sample processing"), this);
 
 	FScopeLock Lock(&CriticalSection);
 
@@ -462,7 +472,7 @@ void FMfMediaTracks::TickAudio(float Rate, FTimespan Time)
 }
 
 
-void FMfMediaTracks::TickInput(float Rate, FTimespan Time, FTimespan DeltaTime)
+void FMfMediaTracks::TickInput(float Rate, FTimespan Time)
 {
 	FScopeLock Lock(&CriticalSection);
 
@@ -474,6 +484,10 @@ void FMfMediaTracks::TickInput(float Rate, FTimespan Time, FTimespan DeltaTime)
 	{
 		CaptionSampleRange = TRange<FTimespan>::AtLeast(Time);
 	}
+	else
+	{
+		// reuse previous range
+	}
 
 	VideoSampleRange = CaptionSampleRange;
 
@@ -722,64 +736,28 @@ bool FMfMediaTracks::SelectTrack(EMediaTrackType TrackType, int32 TrackIndex)
 		return false; // not initialized
 	}
 
+	UE_LOG(LogMfMedia, Verbose, TEXT("Session %p: Selecting %s track %i"), this, *MediaUtils::TrackTypeToString(TrackType), TrackIndex);
+
 	FScopeLock Lock(&CriticalSection);
 
-	int32 OldStreamIndex = INDEX_NONE;
-	int32 NewStreamIndex = INDEX_NONE;
 	int32* SelectedTrack = nullptr;
+	TArray<FTrack>* Tracks = nullptr;
 
 	switch (TrackType)
 	{
 	case EMediaTrackType::Audio:
-		if (TrackIndex == SelectedAudioTrack)
-		{
-			return true; // already selected
-		}
-
-		if ((TrackIndex != INDEX_NONE) && !AudioTracks.IsValidIndex(TrackIndex))
-		{
-			return false; // invalid track
-		}
-
-		UE_LOG(LogMfMedia, Verbose, TEXT("Tracks %llx: Selecting audio track %i instead of %i (%i tracks)"), this, TrackIndex, SelectedAudioTrack, AudioTracks.Num());
-
-		if (AudioTracks.IsValidIndex(SelectedAudioTrack))
-		{
-			OldStreamIndex = AudioTracks[SelectedAudioTrack].StreamIndex;
-		}
-
-		if (TrackIndex != INDEX_NONE)
-		{
-			NewStreamIndex = AudioTracks[TrackIndex].StreamIndex;
-		}
-
 		SelectedTrack = &SelectedAudioTrack;
+		Tracks = &AudioTracks;
 		break;
 
-	case EMediaTrackType::Video:
-		if (TrackIndex == SelectedVideoTrack)
-		{
-			return true; // already selected
-		}
-
-		if ((TrackIndex != INDEX_NONE) && !VideoTracks.IsValidIndex(TrackIndex))
-		{
-			return false; // invalid track
-		}
-
-		UE_LOG(LogMfMedia, Verbose, TEXT("Tracks %llx: Selecting video track %i instead of %i (%i tracks)"), this, TrackIndex, SelectedVideoTrack, VideoTracks.Num());
-
-		if (VideoTracks.IsValidIndex(SelectedVideoTrack))
-		{
-			OldStreamIndex = VideoTracks[SelectedVideoTrack].StreamIndex;
-		}
-
-		if (TrackIndex != INDEX_NONE)
-		{
-			NewStreamIndex = VideoTracks[TrackIndex].StreamIndex;
-		}
+	case EMediaTrackType::Caption:
+		SelectedTrack = &SelectedCaptionTrack;
+		Tracks = &CaptionTracks;
+		break;
 
+	case EMediaTrackType::Video:
 		SelectedTrack = &SelectedVideoTrack;
+		Tracks = &VideoTracks;
 		break;
 
 	default:
@@ -787,47 +765,77 @@ bool FMfMediaTracks::SelectTrack(EMediaTrackType TrackType, int32 TrackIndex)
 	}
 
 	check(SelectedTrack != nullptr);
+	check(Tracks != nullptr);
+
+	if (TrackIndex == *SelectedTrack)
+	{
+		return true; // already selected
+	}
 
-	// apply changes to source reader
-	if (OldStreamIndex != INDEX_NONE)
+	if ((TrackIndex != INDEX_NONE) && !Tracks->IsValidIndex(TrackIndex))
 	{
-		HRESULT Result = SourceReader->SetStreamSelection(OldStreamIndex, FALSE);
+		return false; // invalid track
+	}
+
+	// deselect stream for old track
+	if (*SelectedTrack != INDEX_NONE)
+	{
+		const DWORD StreamIndex = (*Tracks)[*SelectedTrack].StreamIndex;
+		HRESULT Result = PresentationDescriptor->DeselectStream(StreamIndex);
 
 		if (FAILED(Result))
 		{
-			UE_LOG(LogMfMedia, Warning, TEXT("Failed to deselect stream %i: %s"), OldStreamIndex, *MfMedia::ResultToString(Result));
+			UE_LOG(LogMfMedia, Verbose, TEXT("Tracks %p: Failed to deselect stream %i on presentation descriptor: %s"), this, StreamIndex, *MfMedia::ResultToString(Result));
 			return false;
 		}
 
-		Result = SourceReader->Flush(OldStreamIndex);
+		Result = SourceReader->SetStreamSelection(StreamIndex, FALSE);
 
 		if (FAILED(Result))
 		{
-			UE_LOG(LogMfMedia, Warning, TEXT("Failed to flush deselected stream %i: %s"), OldStreamIndex, *MfMedia::ResultToString(Result));
+			UE_LOG(LogMfMedia, Verbose, TEXT("Tracks %p: Failed to deselect stream %i on source reader: %s"), this, StreamIndex, *MfMedia::ResultToString(Result));
+			return false;
 		}
 
-		UE_LOG(LogMfMedia, Verbose, TEXT("Tracks %llx: Disabled stream %i"), this, OldStreamIndex);
+		UE_LOG(LogMfMedia, Verbose, TEXT("Tracks %p: Disabled stream %i"), this, StreamIndex);
 
+		*SelectedTrack = INDEX_NONE;
 		SelectionChanged = true;
+
+		Result = SourceReader->Flush(StreamIndex);
+
+		if (FAILED(Result))
+		{
+			UE_LOG(LogMfMedia, Verbose, TEXT("Tracks %p: Failed to flush deselected stream %i: %s"), this, StreamIndex, *MfMedia::ResultToString(Result));
+		}
 	}
 
-	if (NewStreamIndex != INDEX_NONE)
+	// select stream for new track
+	if (TrackIndex != INDEX_NONE)
 	{
-		const HRESULT Result = SourceReader->SetStreamSelection(NewStreamIndex, TRUE);
+		const DWORD StreamIndex = (*Tracks)[TrackIndex].StreamIndex;
+		HRESULT Result = SourceReader->SetStreamSelection(StreamIndex, TRUE);
 
 		if (FAILED(Result))
 		{
-			UE_LOG(LogMfMedia, Warning, TEXT("Failed to select stream %i: %s"), NewStreamIndex, *MfMedia::ResultToString(Result));
+			UE_LOG(LogMfMedia, Verbose, TEXT("Tracks %p: Failed to select stream %i on source reader: %s"), this, StreamIndex, *MfMedia::ResultToString(Result));
 			return false;
 		}
 
-		UE_LOG(LogMfMedia, Verbose, TEXT("Tracks %llx: Enabled stream %i"), this, NewStreamIndex);
+		Result = PresentationDescriptor->SelectStream(StreamIndex);
+
+		if (FAILED(Result))
+		{
+			UE_LOG(LogMfMedia, Verbose, TEXT("Tracks %p: Failed to select stream %i on presenation descriptor: %s"), this, StreamIndex, *MfMedia::ResultToString(Result));
+			return false;
+		}
 
+		UE_LOG(LogMfMedia, Verbose, TEXT("Tracks %p: Enabled stream %i"), this, StreamIndex);
+			
+		*SelectedTrack = TrackIndex;
 		SelectionChanged = true;
 	}
 
-	*SelectedTrack = TrackIndex;
-
 	return true;
 }
 
@@ -836,76 +844,74 @@ bool FMfMediaTracks::SetTrackFormat(EMediaTrackType TrackType, int32 TrackIndex,
 {
 	FScopeLock Lock(&CriticalSection);
 
-	FTrack* Track = nullptr;
+	TArray<FTrack>* Tracks = nullptr;
 
 	// get track
 	switch (TrackType)
 	{
 	case EMediaTrackType::Audio:
-		if (AudioTracks.IsValidIndex(TrackIndex))
-		{
-			Track = &AudioTracks[TrackIndex];
-		}
+		Tracks = &AudioTracks;
 		break;
 
 	case EMediaTrackType::Caption:
-		if (CaptionTracks.IsValidIndex(TrackIndex))
-		{
-			Track = &CaptionTracks[TrackIndex];
-		}
+		Tracks = &CaptionTracks;
 		break;
 
 	case EMediaTrackType::Video:
-		if (VideoTracks.IsValidIndex(TrackIndex))
-		{
-			Track = &VideoTracks[TrackIndex];
-		}
+		Tracks = &VideoTracks;
 		break;
 
 	default:
-		break; // unsupported track type
+		return false; // unsupported track type
 	};
 
-	if (Track == nullptr)
+	check(Tracks != nullptr);
+
+	if (!Tracks->IsValidIndex(TrackIndex))
 	{
-		return false; // track not found
+		return false; // invalid track
 	}
 
-	if (Track->SelectedFormat == FormatIndex)
+	FTrack& Track = (*Tracks)[TrackIndex];
+
+	if (Track.SelectedFormat == FormatIndex)
 	{
 		return true; // format already set
 	}
 
-	if (!Track->Formats.IsValidIndex(FormatIndex))
+	if (!Track.Formats.IsValidIndex(FormatIndex))
 	{
 		return false; // format not found
 	}
 
 	// get selected format
-	const FFormat& Format = Track->Formats[FormatIndex];
+	const FFormat& Format = Track.Formats[FormatIndex];
 
 	check(Format.InputType.IsValid());
 	check(Format.OutputType.IsValid());
-	check(Track->Handler.IsValid());
+	check(Track.Handler.IsValid());
 
 	// set track format
-	HRESULT Result = Track->Handler->SetCurrentMediaType(Format.InputType);
+	UE_LOG(LogMfMedia, Verbose, TEXT("Tracks %p: Set format %i instead of %i on %s track %i (%i formats)"), this, FormatIndex, Track.SelectedFormat, *MediaUtils::TrackTypeToString(TrackType), TrackIndex, Track.Formats.Num());
+
+	HRESULT Result = Track.Handler->SetCurrentMediaType(Format.InputType);
 
 	if (FAILED(Result))
 	{
-		UE_LOG(LogMfMedia, Warning, TEXT("Failed to set selected media type on handler for stream %i (%s)"), Track->StreamIndex, *MfMedia::ResultToString(Result));
+		UE_LOG(LogMfMedia, Verbose, TEXT("Tracks %p: Failed to set selected media type on handler for stream %i: %s"), this, Track.StreamIndex, *MfMedia::ResultToString(Result));
 		return false;
 	}
 
-	Result = SourceReader->SetCurrentMediaType(Track->StreamIndex, NULL, Format.OutputType);
+	Result = SourceReader->SetCurrentMediaType(Track.StreamIndex, NULL, Format.OutputType);
 
 	if (FAILED(Result))
 	{
-		UE_LOG(LogMfMedia, Warning, TEXT("Failed to set selected media type on reader for stream %i (%s)"), Track->StreamIndex, *MfMedia::ResultToString(Result));
+		UE_LOG(LogMfMedia, Verbose, TEXT("Tracks %p: Failed to set selected media type on reader for stream %i: %s"), this, Track.StreamIndex, *MfMedia::ResultToString(Result));
 		return false;
 	}
 
-	Track->SelectedFormat = FormatIndex;
+	Track.SelectedFormat = FormatIndex;
+	SelectionChanged = true;
 
 	return true;
 }
@@ -914,7 +920,7 @@ bool FMfMediaTracks::SetTrackFormat(EMediaTrackType TrackType, int32 TrackIndex,
 /* FMfMediaTracks implementation
  *****************************************************************************/
 
-void FMfMediaTracks::AddStreamToTracks(uint32 StreamIndex, FString& OutInfo)
+bool FMfMediaTracks::AddStreamToTracks(uint32 StreamIndex, FString& OutInfo)
 {
 	OutInfo += FString::Printf(TEXT("Stream %i\n"), StreamIndex);
 
@@ -926,10 +932,10 @@ void FMfMediaTracks::AddStreamToTracks(uint32 StreamIndex, FString& OutInfo)
 
 		if (FAILED(Result))
 		{
-			UE_LOG(LogMfMedia, Warning, TEXT("Skipping missing stream descriptor for stream %i (%s)"), StreamIndex, *MfMedia::ResultToString(Result));
-			OutInfo += TEXT("    missing stream descriptor\n");
+			UE_LOG(LogMfMedia, Verbose, TEXT("Tracks %p: Failed to get stream descriptor for stream %i: %s"), this, StreamIndex, *MfMedia::ResultToString(Result));
+			OutInfo += TEXT("\tmissing stream descriptor\n");
 
-			return;
+			return false;
 		}
 
 		if (Selected == TRUE)
@@ -945,10 +951,10 @@ void FMfMediaTracks::AddStreamToTracks(uint32 StreamIndex, FString& OutInfo)
 
 		if (FAILED(Result))
 		{
-			UE_LOG(LogMfMedia, Warning, TEXT("Failed to get media type handler for stream %i (%s)"), StreamIndex, *MfMedia::ResultToString(Result));
-			OutInfo += TEXT("    no handler available\n");
+			UE_LOG(LogMfMedia, Verbose, TEXT("Tracks %p: Failed to get media type handler for stream %i: %s"), this, StreamIndex, *MfMedia::ResultToString(Result));
+			OutInfo += TEXT("\tno handler available\n");
 
-			return;
+			return false;
 		}
 	}
 
@@ -959,23 +965,23 @@ void FMfMediaTracks::AddStreamToTracks(uint32 StreamIndex, FString& OutInfo)
 
 		if (FAILED(Result))
 		{
-			UE_LOG(LogMfMedia, Warning, TEXT("Failed to determine major type of stream %i (%s)"), StreamIndex, *MfMedia::ResultToString(Result));
-			OutInfo += TEXT("    failed to determine MajorType\n");
+			UE_LOG(LogMfMedia, Verbose, TEXT("Tracks %p: Failed to determine major type of stream %i: %s"), this, StreamIndex, *MfMedia::ResultToString(Result));
+			OutInfo += TEXT("\tfailed to determine MajorType\n");
 
-			return;
+			return false;
 		}
 
-		UE_LOG(LogMfMedia, Verbose, TEXT("Tracks %llx: Major type of stream %i is %s"), this, StreamIndex, *MfMedia::MajorTypeToString(MajorType));
-		OutInfo += FString::Printf(TEXT("    Type: %s\n"), *MfMedia::MajorTypeToString(MajorType));
+		UE_LOG(LogMfMedia, Verbose, TEXT("Tracks %p: Major type of stream %i is %s"), this, StreamIndex, *MfMedia::MajorTypeToString(MajorType));
+		OutInfo += FString::Printf(TEXT("\tType: %s\n"), *MfMedia::MajorTypeToString(MajorType));
 
 		if ((MajorType != MFMediaType_Audio) &&
 			(MajorType != MFMediaType_SAMI) &&
 			(MajorType != MFMediaType_Video))
 		{
-			UE_LOG(LogMfMedia, Warning, TEXT("Unsupported major type %s for stream %i"), *MfMedia::MajorTypeToString(MajorType), StreamIndex);
-			OutInfo += TEXT("    MajorType is not supported\n");
+			UE_LOG(LogMfMedia, Verbose, TEXT("Tracks %p: Unsupported major type %s of stream %i"), this, *MfMedia::MajorTypeToString(MajorType), StreamIndex);
+			OutInfo += TEXT("\tMajorType is not supported\n");
 
-			return;
+			return false;
 		}
 	}
 
@@ -984,7 +990,7 @@ void FMfMediaTracks::AddStreamToTracks(uint32 StreamIndex, FString& OutInfo)
 	{
 		if (Protected)
 		{
-			OutInfo += FString::Printf(TEXT("    Protected content\n"));
+			OutInfo += FString::Printf(TEXT("\tProtected content\n"));
 		}
 	}
 
@@ -995,10 +1001,10 @@ void FMfMediaTracks::AddStreamToTracks(uint32 StreamIndex, FString& OutInfo)
 
 		if (FAILED(Result))
 		{
-			UE_LOG(LogMfMedia, Warning, TEXT("Failed to get number of track formats in stream %i of type %s"), StreamIndex, *MfMedia::MajorTypeToString(MajorType));
-			OutInfo += TEXT("    failed to get track formats\n");
+			UE_LOG(LogMfMedia, Verbose, TEXT("Tracks %p: Failed to get number of track formats in stream %i"), this, StreamIndex);
+			OutInfo += TEXT("\tfailed to get track formats\n");
 
-			return;
+			return false;
 		}
 	}
 
@@ -1026,9 +1032,11 @@ void FMfMediaTracks::AddStreamToTracks(uint32 StreamIndex, FString& OutInfo)
 	// get current format
 	TComPtr<IMFMediaType> CurrentMediaType;
 	{
-		if (FAILED(Handler->GetCurrentMediaType(&CurrentMediaType)))
+		HRESULT Result = Handler->GetCurrentMediaType(&CurrentMediaType);
+
+		if (FAILED(Result))
 		{
-			UE_LOG(LogMfMedia, Warning, TEXT("Failed to get current media type in stream %i of type %s"), StreamIndex, *MfMedia::MajorTypeToString(MajorType));
+			UE_LOG(LogMfMedia, Verbose, TEXT("Tracks %p: Failed to get current media type in stream %i: %s"), this, StreamIndex, *MfMedia::ResultToString(Result));
 		}
 	}
 
@@ -1039,14 +1047,14 @@ void FMfMediaTracks::AddStreamToTracks(uint32 StreamIndex, FString& OutInfo)
 
 	for (DWORD TypeIndex = 0; TypeIndex < NumMediaTypes; ++TypeIndex)
 	{
-		OutInfo += FString::Printf(TEXT("    Format %i\n"), TypeIndex);
+		OutInfo += FString::Printf(TEXT("\tFormat %i\n"), TypeIndex);
 
 		// get media type
 		TComPtr<IMFMediaType> MediaType;
 
 		if (FAILED(Handler->GetMediaTypeByIndex(TypeIndex, &MediaType)))
 		{
-			OutInfo += TEXT("        failed to get media type\n");
+			OutInfo += TEXT("\t\tfailed to get media type\n");
 
 			continue;
 		}
@@ -1064,22 +1072,22 @@ void FMfMediaTracks::AddStreamToTracks(uint32 StreamIndex, FString& OutInfo)
 
 			if (FAILED(Result))
 			{
-				UE_LOG(LogMfMedia, Warning, TEXT("Failed to get sub-type of format %i in stream %i (%s)"), TypeIndex, StreamIndex, *MfMedia::ResultToString(Result));
-				OutInfo += TEXT("        failed to get sub-type\n");
+				UE_LOG(LogMfMedia, Verbose, TEXT("Tracks %p: Failed to get sub-type of format %i in stream %i: %s"), this, TypeIndex, StreamIndex, *MfMedia::ResultToString(Result));
+				OutInfo += TEXT("\t\tfailed to get sub-type\n");
 
 				continue;;
 			}
 		}
 
 		const FString TypeName = MfMedia::SubTypeToString(SubType);
-		OutInfo += FString::Printf(TEXT("        Codec: %s\n"), *TypeName);
+		OutInfo += FString::Printf(TEXT("\t\tCodec: %s\n"), *TypeName);
 
 		// create output type
 		TComPtr<IMFMediaType> OutputType = MfMedia::CreateOutputType(MajorType, SubType, AllowNonStandardCodecs);
 
 		if (!OutputType.IsValid())
 		{
-			OutInfo += TEXT("        failed to create output type\n");
+			OutInfo += TEXT("\t\tfailed to create output type\n");
 
 			continue;
 		}
@@ -1105,9 +1113,9 @@ void FMfMediaTracks::AddStreamToTracks(uint32 StreamIndex, FString& OutInfo)
 				{ 0 }
 			});
 
-			OutInfo += FString::Printf(TEXT("        Channels: %i\n"), NumChannels);
-			OutInfo += FString::Printf(TEXT("        Sample Rate: %i Hz\n"), SampleRate);
-			OutInfo += FString::Printf(TEXT("        Bits Per Sample: %i\n"), BitsPerSample);
+			OutInfo += FString::Printf(TEXT("\t\tChannels: %i\n"), NumChannels);
+			OutInfo += FString::Printf(TEXT("\t\tSample Rate: %i Hz\n"), SampleRate);
+			OutInfo += FString::Printf(TEXT("\t\tBits Per Sample: %i\n"), BitsPerSample);
 		}
 		else if (MajorType == MFMediaType_SAMI)
 		{
@@ -1127,8 +1135,8 @@ void FMfMediaTracks::AddStreamToTracks(uint32 StreamIndex, FString& OutInfo)
 
 				if (FAILED(Result))
 				{
-					UE_LOG(LogMfMedia, Error, TEXT("Failed to get video output sub-type for stream %i (%s)"), StreamIndex, *MfMedia::ResultToString(Result));
-					OutInfo += FString::Printf(TEXT("        failed to get sub-type"));
+					UE_LOG(LogMfMedia, Verbose, TEXT("Tracks %p: Failed to get video output sub-type for stream %i: %s"), this, StreamIndex, *MfMedia::ResultToString(Result));
+					OutInfo += FString::Printf(TEXT("\t\tfailed to get sub-type"));
 
 					continue;
 				}
@@ -1140,12 +1148,12 @@ void FMfMediaTracks::AddStreamToTracks(uint32 StreamIndex, FString& OutInfo)
 			{
 				if (SUCCEEDED(::MFGetAttributeSize(MediaType, MF_MT_FRAME_SIZE, (UINT32*)&OutputDim.X, (UINT32*)&OutputDim.Y)))
 				{
-					OutInfo += FString::Printf(TEXT("        Dimensions: %i x %i\n"), OutputDim.X, OutputDim.Y);
+					OutInfo += FString::Printf(TEXT("\t\tDimensions: %i x %i\n"), OutputDim.X, OutputDim.Y);
 				}
 				else
 				{
 					OutputDim = FIntPoint::ZeroValue;
-					OutInfo += FString::Printf(TEXT("        Dimensions: n/a\n"));
+					OutInfo += FString::Printf(TEXT("\t\tDimensions: n/a\n"));
 				}
 			}
 
@@ -1157,12 +1165,12 @@ void FMfMediaTracks::AddStreamToTracks(uint32 StreamIndex, FString& OutInfo)
 				if (SUCCEEDED(::MFGetAttributeRatio(MediaType, MF_MT_FRAME_RATE, &Numerator, &Denominator)))
 				{
 					FrameRate = static_cast<float>(Numerator) / Denominator;
-					OutInfo += FString::Printf(TEXT("        Frame Rate: %g fps\n"), FrameRate);
+					OutInfo += FString::Printf(TEXT("\t\tFrame Rate: %g fps\n"), FrameRate);
 				}
 				else
 				{
 					FrameRate = 0.0f;
-					OutInfo += FString::Printf(TEXT("        Frame Rate: n/a\n"));
+					OutInfo += FString::Printf(TEXT("\t\tFrame Rate: n/a\n"));
 				}
 			}
 
@@ -1192,7 +1200,12 @@ void FMfMediaTracks::AddStreamToTracks(uint32 StreamIndex, FString& OutInfo)
 					FrameRates = TRange<float>(FrameRate);
 				}
 
-				OutInfo += FString::Printf(TEXT("        Frame Rate Range: %g - %g fps\n"), FrameRates.GetLowerBoundValue(), FrameRates.GetUpperBoundValue());
+				OutInfo += FString::Printf(TEXT("\t\tFrame Rate Range: %g - %g fps\n"), FrameRates.GetLowerBoundValue(), FrameRates.GetUpperBoundValue());
+
+				if (FrameRates.IsDegenerate() && (FrameRates.GetLowerBoundValue() == 1.0f))
+				{
+					OutInfo += FString::Printf(TEXT("\t\tpossibly a still image stream (may not work)\n"));
+				}
 			}
 
 			FIntPoint BufferDim;
@@ -1335,12 +1348,36 @@ void FMfMediaTracks::AddStreamToTracks(uint32 StreamIndex, FString& OutInfo)
 			}
 			else
 			{
-				UE_LOG(LogMfMedia, Warning, TEXT("Failed to set current media type on reader for stream %i (%s)"), Track->StreamIndex, *MfMedia::ResultToString(Result));
+				UE_LOG(LogMfMedia, Verbose, TEXT("Tracks %p: Failed to set current media type on reader for stream %i: %s"), this, Track->StreamIndex, *MfMedia::ResultToString(Result));
+			}
+		}
+	}
+
+	// ensure that a track format is selected
+	if (Track->SelectedFormat == INDEX_NONE)
+	{
+		for (int32 FormatIndex = 0; FormatIndex < Track->Formats.Num(); ++FormatIndex)
+		{
+			const FFormat& Format = Track->Formats[FormatIndex];
+			const HRESULT ReaderResult = SourceReader->SetCurrentMediaType(StreamIndex, NULL, Format.OutputType);
+			const HRESULT TrackResult = Handler->SetCurrentMediaType(Format.InputType);
+
+			if (SUCCEEDED(TrackResult) && SUCCEEDED(ReaderResult))
+			{
+				UE_LOG(LogMfMedia, Verbose, TEXT("Tracks %p: Picked default format %i for stream %i"), this, FormatIndex, StreamIndex);
+				Track->SelectedFormat = FormatIndex;
+				break;
 			}
 		}
+
+		if (Track->SelectedFormat == INDEX_NONE)
+		{
+			UE_LOG(LogMfMedia, Verbose, TEXT("Tracks %p: No supported media types found in stream %i"), this, StreamIndex);
+			OutInfo += TEXT("\tunsupported media type\n");
+		}
 	}
 
-	// get stream info
+	// set track details
 	PWSTR OutString = NULL;
 	UINT32 OutLength = 0;
 
@@ -1365,17 +1402,7 @@ void FMfMediaTracks::AddStreamToTracks(uint32 StreamIndex, FString& OutInfo)
 	Track->Protected = Protected;
 	Track->StreamIndex = StreamIndex;
 
-	if (Track->Formats.Num() > 0)
-	{
-		Track->SelectedFormat = 0;
-	}
-	else
-	{
-		Track->SelectedFormat = INDEX_NONE;
-
-		UE_LOG(LogMfMedia, Warning, TEXT("No supported media type in stream %i of type %s"), StreamIndex, *MfMedia::MajorTypeToString(MajorType));
-		OutInfo += TEXT("    unsupported media type\n");
-	}
+	return true;
 }
 
 
@@ -1456,7 +1483,7 @@ bool FMfMediaTracks::RequestSample(DWORD StreamIndex)
 
 	if (FAILED(Result))
 	{
-		UE_LOG(LogMfMedia, VeryVerbose, TEXT("Tracks %llx: Failed to request sample for stream %i: %s"), this, StreamIndex, *MfMedia::ResultToString(Result));
+		UE_LOG(LogMfMedia, VeryVerbose, TEXT("Tracks %p: Failed to request sample for stream %i: %s"), this, StreamIndex, *MfMedia::ResultToString(Result));
 		return false;
 	}
 
@@ -1492,7 +1519,7 @@ void FMfMediaTracks::UpdateAudio()
 	}
 
 	#if MFMEDIATRACKS_TRACE_SAMPLES
-		UE_LOG(LogMfMedia, VeryVerbose, TEXT("Tracks %llx: Requesting audio sample"), this);
+		UE_LOG(LogMfMedia, VeryVerbose, TEXT("Tracks %p: Requesting audio sample"), this);
 	#endif
 
 	AudioSamplePending = RequestSample(AudioTracks[SelectedAudioTrack].StreamIndex);
@@ -1517,7 +1544,7 @@ void FMfMediaTracks::UpdateCaptions()
 	}
 
 	#if MFMEDIATRACKS_TRACE_SAMPLES
-		UE_LOG(LogMfMedia, VeryVerbose, TEXT("Tracks %llx: Requesting caption sample"), this);
+		UE_LOG(LogMfMedia, VeryVerbose, TEXT("Tracks %p: Requesting caption sample"), this);
 	#endif
 
 	CaptionSamplePending = RequestSample(CaptionTracks[SelectedCaptionTrack].StreamIndex);
@@ -1542,7 +1569,7 @@ void FMfMediaTracks::UpdateVideo()
 	}
 
 	#if MFMEDIATRACKS_TRACE_SAMPLES
-		UE_LOG(LogMfMedia, VeryVerbose, TEXT("Tracks %llx: Requesting video sample"), this);
+		UE_LOG(LogMfMedia, VeryVerbose, TEXT("Tracks %p: Requesting video sample"), this);
 	#endif
 
 	VideoSamplePending = RequestSample(VideoTracks[SelectedVideoTrack].StreamIndex);
diff --git a/Engine/Plugins/Media/MfMedia/Source/MfMedia/Private/Player/MfMediaTracks.h b/Engine/Plugins/Media/MfMedia/Source/MfMedia/Private/Player/MfMediaTracks.h
--- a/Engine/Plugins/Media/MfMedia/Source/MfMedia/Private/Player/MfMediaTracks.h
+++ b/Engine/Plugins/Media/MfMedia/Source/MfMedia/Private/Player/MfMediaTracks.h
@@ -199,10 +199,9 @@ public:
 	 *
 	 * @param Rate The current playback rate.
 	 * @param Time The current presentation time.
-	 * @param DeltaTime Time since the last tick.
 	 * @see TickAudio
 	 */
-	void TickInput(float Rate, FTimespan Time, FTimespan DeltaTime);
+	void TickInput(float Rate, FTimespan Time);
 
 public:
 
@@ -227,8 +226,9 @@ protected:
 	 *
 	 * @param StreamIndex The index number of the stream in the presentation descriptor.
 	 * @param OutInfo String to append debug information to.
+	 * @return true on success, false otherwise.
 	 */
-	void AddStreamToTracks(uint32 StreamIndex, FString& OutInfo);
+	bool AddStreamToTracks(uint32 StreamIndex, FString& OutInfo);
 
 	/**
 	 * Get the specified audio format.
diff --git a/Engine/Plugins/Media/WmfMedia/Source/WmfMedia/Private/Player/WmfMediaPlayer.cpp b/Engine/Plugins/Media/WmfMedia/Source/WmfMedia/Private/Player/WmfMediaPlayer.cpp
--- a/Engine/Plugins/Media/WmfMedia/Source/WmfMedia/Private/Player/WmfMediaPlayer.cpp
+++ b/Engine/Plugins/Media/WmfMedia/Source/WmfMedia/Private/Player/WmfMediaPlayer.cpp
@@ -79,7 +79,7 @@ FString FWmfMediaPlayer::GetInfo() const
 }
 
 
-FName FWmfMediaPlayer::GetName() const
+FName FWmfMediaPlayer::GetPlayerName() const
 {
 	static FName PlayerName(TEXT("WmfMedia"));
 	return PlayerName;
@@ -140,13 +140,13 @@ bool FWmfMediaPlayer::Open(const TSharedRef<FArchive, ESPMode::ThreadSafe>& Arch
 
 	if (Archive->TotalSize() == 0)
 	{
-		UE_LOG(LogWmfMedia, Error, TEXT("Cannot open media from archive (archive is empty)."));
+		UE_LOG(LogWmfMedia, Verbose, TEXT("Player %p: Cannot open media from archive (archive is empty)"), this);
 		return false;
 	}
 
 	if (OriginalUrl.IsEmpty())
 	{
-		UE_LOG(LogWmfMedia, Error, TEXT("Cannot open media from archive (no original URL provided)."));
+		UE_LOG(LogWmfMedia, Verbose, TEXT("Player %p: Cannot open media from archive (no original URL provided)"), this);
 		return false;
 	}
 
@@ -154,39 +154,37 @@ bool FWmfMediaPlayer::Open(const TSharedRef<FArchive, ESPMode::ThreadSafe>& Arch
 }
 
 
-void FWmfMediaPlayer::TickFetch(FTimespan /*DeltaTime*/)
+void FWmfMediaPlayer::TickFetch(FTimespan /*DeltaTime*/, FTimespan /*Timecode*/)
 {
 	bool MediaSourceChanged = false;
-	bool TopologyChanged = false;
+	bool TrackSelectionChanged = false;
 
-	Tracks->GetFlags(MediaSourceChanged, TopologyChanged);
+	Tracks->GetFlags(MediaSourceChanged, TrackSelectionChanged);
 
 	if (MediaSourceChanged)
 	{
 		EventSink.ReceiveMediaEvent(EMediaEvent::TracksChanged);
 	}
 
-	if (TopologyChanged)
+	if (TrackSelectionChanged)
 	{
-		if (Tracks->IsInitialized())
-		{
-			Session->SetTopology(Tracks->CreateTopology(), Tracks->GetDuration());
-		}
-		else
+		UE_LOG(LogWmfMedia, Verbose, TEXT("Player %p: Creating and setting new playback topology"), this);
+
+		if (!Tracks->IsInitialized() || !Session->SetTopology(Tracks->CreateTopology(), Tracks->GetDuration()))
 		{
 			Session->Shutdown();
 			EventSink.ReceiveMediaEvent(EMediaEvent::MediaOpenFailed);
 		}
 	}
 
-	if (MediaSourceChanged || TopologyChanged)
+	if (MediaSourceChanged || TrackSelectionChanged)
 	{
 		Tracks->ClearFlags();
 	}
 }
 
 
-void FWmfMediaPlayer::TickInput(FTimespan DeltaTime)
+void FWmfMediaPlayer::TickInput(FTimespan DeltaTime, FTimespan /*Timecode*/)
 {
 	// forward session events
 	TArray<EMediaEvent> OutEvents;
@@ -219,14 +217,14 @@ bool FWmfMediaPlayer::InitializePlayer(const TSharedPtr<FArchive, ESPMode::Threa
 	// initialize presentation on a separate thread
 	const EAsyncExecution Execution = Precache ? EAsyncExecution::Thread : EAsyncExecution::ThreadPool;
 
-	Async<void>(Execution, [Archive, Url, Precache, StreamsPtr = TWeakPtr<FWmfMediaTracks, ESPMode::ThreadSafe>(Tracks)]()
+	Async<void>(Execution, [Archive, Url, Precache, TracksPtr = TWeakPtr<FWmfMediaTracks, ESPMode::ThreadSafe>(Tracks)]()
 	{
-		TSharedPtr<FWmfMediaTracks, ESPMode::ThreadSafe> PinnedStreams = StreamsPtr.Pin();
+		TSharedPtr<FWmfMediaTracks, ESPMode::ThreadSafe> PinnedTracks = TracksPtr.Pin();
 
-		if (PinnedStreams.IsValid())
+		if (PinnedTracks.IsValid())
 		{
 			TComPtr<IMFMediaSource> MediaSource = WmfMedia::ResolveMediaSource(Archive, Url, Precache);
-			PinnedStreams->Initialize(MediaSource);
+			PinnedTracks->Initialize(MediaSource);
 		}
 	});
 
