Back out revision 5 from //UE4/Main/Engine/Plugins/Media/MfMedia/MfMedia.uplugin
- marks plugin as disabled by default
- removes Win32/64 whitelists

#tests preflighted "Test Orion Win64"

diff --git a/Engine/Plugins/Media/MfMedia/MfMedia.uplugin b/Engine/Plugins/Media/MfMedia/MfMedia.uplugin
--- a/Engine/Plugins/Media/MfMedia/MfMedia.uplugin
+++ b/Engine/Plugins/Media/MfMedia/MfMedia.uplugin
@@ -10,7 +10,7 @@
 	"DocsURL" : "https://forums.unrealengine.com/showthread.php?46879-Media-Framework-Documentation-for-4-5-Preview",
 	"MarketplaceURL" : "",
 	"SupportURL" : "",
-	"EnabledByDefault" : true,
+	"EnabledByDefault" : false,
 	"CanContainContent" : false,
 	"IsBetaVersion" : false,
 	"Installed" : false,
@@ -20,7 +20,7 @@
 			"Name" : "MfMedia",
 			"Type" : "RuntimeNoCommandlet",
 			"LoadingPhase" : "PreLoadingScreen",
-			"WhitelistPlatforms" : [ "XboxOne", "Win32", "Win64" ]
+			"WhitelistPlatforms" : [ "XboxOne" ]
 		},
 		{
 			"Name" : "MfMediaEditor",
@@ -36,7 +36,7 @@
 			"Name" : "MfMediaFactory",
 			"Type" : "RuntimeNoCommandlet",
 			"LoadingPhase" : "PostEngineInit",
-			"WhitelistPlatforms" : [ "XboxOne", "Win32", "Win64" ]
+			"WhitelistPlatforms" : [ "XboxOne" ]
 		}
 	]
 }
