Copying //UE4/Dev-Sequencer to //UE4/Dev-Main (Source: //UE4/Dev-Sequencer @ 3800792)

============================
  MAJOR FEATURES & CHANGES
============================

Change 3631172 by [REDACTED_PERSON]

	Sequence Recorder: Record to the specified target animation for the target actor only. Newly tracked components will have newly created animations so that they don't record to the same target animation assets.

Change 3632895 by [REDACTED_PERSON]

	Sequencer: Fixed crash caused by erroneous persistent references to FSequencer.

Change 3634864 by [REDACTED_PERSON]

	Sequencer: Bake transforms

Change 3634879 by [REDACTED_PERSON]

	Sequence Recorder: Close the target animation asset editor if it exists before recording into it.

Change 3650553 by [REDACTED_PERSON]

	Sequencer: Switch shot display name to FString so that it's not localized.

Change 3650848 by [REDACTED_PERSON]

	Linear Timecode Plugin - first pass - still needs output interface

Change 3653174 by [REDACTED_PERSON]

	Sequencer: Allow blending for vector tracks

Change 3653205 by [REDACTED_PERSON]

	Sequencer: Fix default burnin content.

Change 3653570 by [REDACTED_PERSON]

	Win32 Build fix: EmptyLinkFunctionForStaticInitializationLinearTimecode

Change 3658673 by [REDACTED_PERSON]

	Timecode - Adding Accessors, renaming files

Change 3663633 by [REDACTED_PERSON]

	Sequencer: Add option to export an object's transform to a camera anim asset.

Change 3665971 by [REDACTED_PERSON]

	Sequence Recorder: Add an active column for sequence recorder

Change 3666066 by [REDACTED_PERSON]

	Sequencer: Suspend broadcast of selection delegates when performing multiple operations.

Change 3671057 by [REDACTED_PERSON]

	PR #4055: Fixed ACameraRig_Rail PreviewMeshes to match Spline Component (Contributed by [REDACTED_PERSON])

Change 3677104 by [REDACTED_PERSON]

	Sequencer: Add transaction for easing length.

Change 3677105 by [REDACTED_PERSON]

	Sequencer: Restore pre-animated state when changing active channels.

Change 3681069 by [REDACTED_PERSON]

	LinearTimecode: Processing samples only if they have been dequeued successfully

Change 3681783 by [REDACTED_PERSON]

	Control Rig: Fixed crash attempting to zero out array with the wrong size

Change 3681949 by [REDACTED_PERSON]

	OnTimecodeChange Delegate, Comments in DropTimecode, DropTimecode to string auto conversion

Change 3682222 by [REDACTED_PERSON]

	EngineTest: Enabled MP4 Media Player test in the level ShaderModels. Set the texture ShaderModelPlayer_Video to use the correct media player.

Change 3683523 by [REDACTED_PERSON]

	Sequencer: Display sub sequence name on binding ID pickers rather than the sub section name

Change 3684318 by [REDACTED_PERSON]

	Sequencer: Fix crash in dragging a level sequence into the tree area.

Change 3687609 by [REDACTED_PERSON]

	Curve Editor: Add ModifyOwnerChange to CurveOwnerInterface. Call mark as changed when modifying keys or tangents.

Change 3687617 by [REDACTED_PERSON]

	Sequencer: Move some commands out of the generic sequencer command bindings so that they don't take over the viewport. ie. End for "Snap to Floor" should still function in the viewport.

Change 3687969 by [REDACTED_PERSON]

	Sequencer: New compile on the fly logic
	  - Sequencer is now able to compile partially or completely out-of-date evaluation templates from the source data as it needs. This affords much more efficient compilation when working within sequencer.
	  - Added the concept of 'instance data' for sub sequences, available through the IMovieScenePlayer interface or persistent data stores. This replaces the compilation of specific templates for control rig templates.
	  - Moved sub tracks and sections to MovieScene module
	  - Removed the concept of shared tracks. Any previous uses should port over to shared execution tokens instead.

Change 3689884 by [REDACTED_PERSON]

	Sequencer: Remove null tracks on object bindings. Tracks can become null if they're from a plugin and the plugin is disabled.

Change 3690130 by [REDACTED_PERSON]

	Sequencer: Fixed UMG animations not working with blueprint nativization

Change 3690162 by [REDACTED_PERSON]

	Sequencer: Fix assert when using filler shots

Change 3691904 by [REDACTED_PERSON]

	Sequencer: Fix overlap by looking at the key behind as well. This fixes a bug if there's 3 keyframes on consecutive frames, if you zoom out, you should see two bordered keys when the overlap threshold is passed.

Change 3692149 by [REDACTED_PERSON]

	Sequencer: Change FSortByPrioritySegmentBlender::Blend to fall back to flag comparison if any sections are null for automation tests.
	  - Combinations of null/non-null will never happen in the wild because we either run with sections or without, but the static analysis doesn't know that

Change 3692212 by [REDACTED_PERSON]

	Core: Changed TRange to pass-by-value for small types, inlined some functions

Change 3693044 by [REDACTED_PERSON]

	Sequencer: Update pivot location for selection when closing sequencer.

Change 3693061 by [REDACTED_PERSON]

	Sequencer: Remove force evaluate on end scrubbing. Backout previous fix because 3657907 is the real fix.

Change 3694308 by [REDACTED_PERSON]

	Sequencer: Sub tracks now blend correctly again

Change 3694365 by [REDACTED_PERSON]

	Sequencer: Fixed exponential slow down when compiling large quantities of sub sequences on cook caused by not clamping sub-sequence compilation correctly

Change 3698176 by [REDACTED_PERSON]

	Sequencer: Enable the camera cut track when popping back to the master only if there's a camera cut track in the master. This fixes an issue where if you don't have a camera cut track in the master, the camera gets locked to a camera cut in a subscene and you can't toggle out of it.

Change 3698979 by [REDACTED_PERSON]

	Sequencer: Fix delay before warmup causing sequences not to render.

Change 3700149 by [REDACTED_PERSON]

	Sequencer: Fix delay at shot boundaries causing sequences to not play back and render out. The shot id needs to be tracked to determine whether a new shot is encountered.

Change 3705624 by [REDACTED_PERSON]

	Sequencer: Reset drop node on drag leave

Change 3718781 by [REDACTED_PERSON]

	MediaAssets: Removed obsolete pointer checks (Playlist is always valid); added checks to enforce this invariant

Change 3721586 by [REDACTED_PERSON]

	MediaAssets: Added verbose logging to MediaPlayer BP functions

Change 3721662 by [REDACTED_PERSON]

	WmfMedia: Fixed media session not generating PlaybackEndReached event when session forced to stop

Change 3722954 by [REDACTED_PERSON]

	MediaAssets: Removed looping option from play lists

Change 3723076 by [REDACTED_PERSON]

	Removed controls for looping playlist.

Change 3724202 by [REDACTED_PERSON]

	Sequencer: Spawn time optimizations

Change 3725359 by [REDACTED_PERSON]

	Sequencer: Fix spawnables not playing back. bIsEditorPreviewActor = false for sequencer spawnables so that BeginPlay() doesn't get skipped.

Change 3725362 by [REDACTED_PERSON]

	Sequencer: Change check to ensure for null object in property actuate. This prevents a crash when there's a track without a binding object. This occur when pasting a property track to the root.

Change 3725802 by [REDACTED_PERSON]

	Sequencer: Add the ability to drag sections up.

Change 3726520 by [REDACTED_PERSON]

	MediaAssets: Fixed OnMediaOpened called when looping a source

Change 3728509 by [REDACTED_PERSON]

	Sequencer: Fix crash pasting a camera cut track because the outer isn't set to the movie scene.

Change 3729216 by [REDACTED_PERSON]

	Sequencer: Fix pre edit property chain broadcast so that the property path will include possible struct/array node.

Change 3729248 by [REDACTED_PERSON]

	Sequencer: Add default expansion states to allow track editors to specify them per track type.

	Material track is currently the only track that defaults to expanded.

Change 3729392 by [REDACTED_PERSON]

	ImgMedia: Disabling video output when video track is selected

Change 3729452 by [REDACTED_PERSON]

	MediaPlayerEditor: Showing a notification if no video track available or selected

Change 3729462 by [REDACTED_PERSON]

	MediaPlayerEditor: Creating overlay font only once

Change 3731155 by [REDACTED_PERSON]

	Sequencer: Set row index when creating a new take

Change 3731446 by [REDACTED_PERSON]

	Sequencer: Add null check and warning for invalid GetParameterCollectionInstance.

Change 3731640 by [REDACTED_PERSON]

	Sequencer: Automatically create a camera cut track if a camera is dropped and there's no existing camera cut or there's no existing camera cut sections.

Change 3731653 by [REDACTED_PERSON]

	Sequencer: Defer details panel updates on scrubbing and playing.

Change 3732597 by [REDACTED_PERSON]

	Sequencer: Don't swap the camera cut section if it already exists and a camera is created.

Change 3733554 by [REDACTED_PERSON]

	Sequencer: Fix crash when existing hotspot is null.

Change 3734305 by [REDACTED_PERSON]

	MediaAssets: Fixed assertion when assigning same media player to a media texture twice

Change 3735735 by [REDACTED_PERSON]

	MediaFrameworkTest: Setting MediaAssets log category Verbose to prevent log spam in Editor

Change 3736393 by [REDACTED_PERSON]

	Sequencer: Set min/max values for generic key area so that they don't default to 0,10.

Change 3736804 by [REDACTED_PERSON]

	Sequencer: Fix rotation rotator->quat->rotator conversion which prevents you from typing in a rotation of 0,0,320 into the key editor. Unwind rotations from the previous transform to the current transform so that the nearest rotation is set rather than the rotator->quat->rotator conversion.

Change 3737790 by [REDACTED_PERSON]

	Sequencer: Follow up fix to quat->rot->quat so that unwinding is only processed if the last transform is valid.

Change 3739738 by [REDACTED_PERSON]

	PR #4184: Unregister Missing TrackEditors in MovieSceneTools (Contributed by [REDACTED_PERSON])

Change 3741910 by [REDACTED_PERSON]

	MfMedia: Simplified track switching code; enabled seeking

Change 3742171 by [REDACTED_PERSON]

	MediaAssets: Fixed timecode check to include current presentation time

Change 3744889 by [REDACTED_PERSON]

	Sequencer: Use drag and drop actor factory if specifiied. This ensures that the correct actor factory is used in creating the object template for the sequencer spawner. This fixes some spawnables not getting created properly (ie. Empty Actor)

Change 3746049 by [REDACTED_PERSON]

	FBX Import: Fixed tangent overrides producing the wrong type of key tangent modes

Change 3747870 by [REDACTED_PERSON]

	Sequencer: Update auto scroll when moving keys/sections.

Change 3747885 by [REDACTED_PERSON]

	UMG: Remove template instance finish on pause. It appears that was only there to stop audio (CL #3161056). Instead, evaluate in the stopped position with the stopped playback state.

Change 3747914 by [REDACTED_PERSON]

	Sequencer: Add config for default completion mode for movie scene sequences. The default for level sequences is RestoreState. All others, such as UMG are set to KeepState.

Change 3747915 by [REDACTED_PERSON]

	Sequencer: Set default completion mode for all sections to project default.

Change 3748377 by [REDACTED_PERSON]

	EngineTest: Fixed media player automation test BP

Change 3748531 by [REDACTED_PERSON]

	EngineTest: Fixed media player not opening media during FTest setup

	I fixed up the Blueprint. Also, don't use PlayOnOpen on the media player if you don't intend to play the movie after opening. You can call Pause right after opening. If you enable PlayOnOpen, the order of operations will not work, because it will first call Pause via the IsReady check in Tick and then Play via the PlayOnOpen checkbox. I disabled PlayOnOpen.

Change 3750362 by [REDACTED_PERSON]

	PR #4205: parameter not used bug fixes (Contributed by [REDACTED_PERSON])

Change 3750885 by [REDACTED_PERSON]

	Sequencer: Prevent throttling on the curve editor so that editor world tick can apply.

Change 3757123 by [REDACTED_PERSON]

	MediaAssets: Allowing media players to be part of GC clusters

Change 3764055 by [REDACTED_PERSON]

	Sequencer: Add SetPlaybackPosition as latent action to fix crash when calling SetPlaybackPosition in an event.

Change 3764555 by [REDACTED_PERSON]

	Sequencer: Change camera cut binding from guid to object binding id to allow for subscene bindings

Change 3764613 by [REDACTED_PERSON]

	Sequencer: Show binding ID picker customization on all details panels. Allow creation of new camera cut sections from existing bindings.

	Deprecate AddNewCameraCut(FGuid)

Change 3766389 by [REDACTED_PERSON]

	Sequencer: Don't override existing camera binding id on postload if it's valid.

Change 3766432 by [REDACTED_PERSON]

	Sequencer: Switch constraint guid to constraint binding ID

Change 3767160 by [REDACTED_PERSON]

	WmfMedia: Added more logging

Change 3767843 by [REDACTED_PERSON]

	MediaAssets: Fixed Crash in UMediaSoundComponent::UpdatePlayer

Change 3771038 by [REDACTED_PERSON]

	Engine: Changed message log back to ensure

	Cannot use message log on render thread as it may call LoadModule

Change 3771071 by [REDACTED_PERSON]

	MediaCompositing: Fixed default media plane material

Change 3771287 by [REDACTED_PERSON]

	Engine: Changed ensure to global log

Change 3771450 by [REDACTED_PERSON]

	MediaAssets: Fixed crash when dragging media player viewport

Change 3771461 by [REDACTED_PERSON]

	ImgMedia: Fixed seeking not working when paused

Change 3772999 by [REDACTED_PERSON]

	Sequencer: Added the ability to supply an instance data object to sequence runtime evaluation
	  - This allows systems to inject information into a sequence at runtime, enabling more dynamic control of tracks
	  - Level sequence actors use this by default to supply a dynamic 'transform origin' to all component transform tracks, to which all absolute transform sections will be added.

Change 3774633 by [REDACTED_PERSON]

	MediaAssets: Fixed MediaPlayer.Prev/Next not playing automatically if player was playing previously

Change 3774961 by [REDACTED_PERSON]

	UdpMessaging: Patch to make message serialization notification thread-safe

Change 3778108 by [REDACTED_PERSON]

	Sequencer: Scroll in track area.

Change 3778270 by [REDACTED_PERSON]

	Sequencer: Copy/cut/paste/duplicate object bindings

	FMovieSceneCopyableBinding - struct containing the copied FMovieSceneBinding, FMovieScenePossessable and FMovieSceneSpawnable
	MovieScene - AddPossessable, AddSpawnable with existing possessable or spawnable
	ISequencer - FOnMovieSceneBindingsPasted to allow fixing up bindings on paste

Change 3778674 by [REDACTED_PERSON]

	Engine: Fixed online documentation for external texture registry

Change 3778699 by [REDACTED_PERSON]

	Engine: Fixed external textures not registering correctly in all cases

Change 3778776 by [REDACTED_PERSON]

	Engine: Corrected code comments for external texture registry

Change 3779153 by [REDACTED_PERSON]

	LinearTimecode: Fixed plugin category

Change 3779166 by [REDACTED_PERSON]

	WmfMedia: Reversing FourceCC display string, so it displays correctly

Change 3780576 by [REDACTED_PERSON]

	Sequencer: Fix viewport invalidation so that it happens on sequence evaluation.

Change 3780619 by [REDACTED_PERSON]

	Sequencer: Always save default spawnable state regardless of focused sequence. This fixes a bug where if you step back to the master sequence (and the spawnable still exists), and then scrub outside the region where the spawnable exists, it gets destroyed but saved default spawnable state doesn't get called because it's no longer the focused sequence.

Change 3781227 by [REDACTED_PERSON]

	Sequencer: Add playback speed settings menu

Change 3782213 by [REDACTED_PERSON]

	WmfMedia: Fixed H.265 frames being dropped due to false negative buffer size check

Change 3782239 by [REDACTED_PERSON]

	Sequencer: Add shortcut to tooltips

Change 3782388 by [REDACTED_PERSON]

	WmfMedia: Fixed IYUV encoded AVI not playing correctly

Change 3782541 by [REDACTED_PERSON]

	Protect against asking for duration before prepare completed in movie player
	#ue4
	#android

Change 3787676 by [REDACTED_PERSON]

	Sequencer: Export unmatched float properties to custom attributes

Change 3787788 by [REDACTED_PERSON]

	Extras: Natvis for FTimespan

Change 3789419 by [REDACTED_PERSON]

	Sequencer: Fix import camera so that when new cameras are created, values from the fbx are going only to the newly created cameras.

	Also, added bReduceKeys and bReduceKeysTolerance to import fbx.

Change 3789628 by [REDACTED_PERSON]

	Sequencer: Remove Shift-C and Shift-D as a hotkey for toggling the cinematic viewport as it was causing some confusion when users accidentally hit it.

Change 3789716 by [REDACTED_PERSON]

	Seqeuncer: Fixed slot animation not being restored for montages that are recreated during evaluation

	  - The issue here is that the animation system can deem montages as inactive at any point, and will re-create a new montage instance if this is the case. Previously, we stored pre-animated state on the anim instance, and cached the montage that we played, which meant that we attempt to restore a stale montage.

	  - The solution is to store pre-animated state for the montages, and cache the anim instance. This means we'll always store pre-animated state if the montage is re-created.

Change 3789874 by [REDACTED_PERSON]

	Sequencer: Add option to only show selected nodes only

	ISequencer - Added OnGetIsBindingVisible(), OnGetIsTrackVisible()
	LevelEditorSequencerIntegration - implement OnGetIsBindingVisible based on actor selection

Change 3789877 by [REDACTED_PERSON]

	Control Rig: Implement IsTrackVisible to show nodes in the sequencer tree that are selected.

Change 3791514 by [REDACTED_PERSON]

	Image Plate: Fixed image plates not showing up in high res screenshots
	  - The image plate mesh had a LOD index of -1 which meant it never showed up when a LOD to render was specified
	  - Also simplified the bounds calculation

Change 3791596 by [REDACTED_PERSON]

	Sequence Recorder: Don't record editor only components

Change 3791659 by [REDACTED_PERSON]

	Sequencer: Event struct payloads now support user defined struct default values on creation

Change 3791718 by [REDACTED_PERSON]

	Sequence Recorder: Fix level sequences not triggering when recording. Level sequences would not get recorded if the world settings actor was not recorded.

Change 3791761 by [REDACTED_PERSON]

	Sequencer: Find available non-overlapping row index when adding subsections.

Change 3791837 by [REDACTED_PERSON]

	Sequencer: Fix sequencer not updating after toggling Bind Sequencer to PIE/Simulate while PIE is active.

Change 3791875 by [REDACTED_PERSON]

	Sequence Recorder: Add recording button now adds for any selected actors.

Change 3791880 by [REDACTED_PERSON]

	Sequence Recorder: Fix end pie not stopping sequence recording if not recording world settings.

Change 3791932 by [REDACTED_PERSON]

	Sequence Recorder: Record spawned actors immediately so that they won't be missed if they're deleted before tick.  Clean up queued recordings that were from spawned actors.

Change 3791954 by [REDACTED_PERSON]

	Sequence Recorder: Don't process if not recording on end pie.

Change 3791955 by [REDACTED_PERSON]

	Sequence Recorder: Fix recording spawnables not getting the correct position for being spawned at.

	Set a transform on the movie scene spawnable which will be used if the root component of the object template doesn't exist.

Change 3792061 by [REDACTED_PERSON]

	Sequence Recorder: Change FOnRecordingStarted and FOnRecordingFinished to multicast delegate.

	If a sequence being recorded into is currently open in sequencer, it gets reloaded to prevent crashes.

Change 3792271 by [REDACTED_PERSON]

	Sequence Recorder: Duplicate the actor to trigger before playing so that a sequence can be recorded and played back at the same time.

Change 3792292 by [REDACTED_PERSON]

	MfMedia: Fixed tracks being reported in reverse order

Change 3792811 by [REDACTED_PERSON]

	Sequencer: Add option to disable camera cuts on the movie scene player.

Change 3793341 by [REDACTED_PERSON]

	AvfMedia: Fix for crash when opening AVF Media after opening multiple players

	Merged from shelved CL in Dev-Rendering

Change 3793642 by [REDACTED_PERSON]

	Sequencer: Delete and rename icons.

	PR #4194: Improvements to Sequencer (Contributed by [REDACTED_PERSON])

Change 3794678 by [REDACTED_PERSON]

	Sequencer: Fix crash on HandleRecordingFinished

Change 3795007 by [REDACTED_PERSON]

	AudioMixer: Fix for AVF Media Player audio is garbled and has static

Change 3795103 by [REDACTED_PERSON]

	Sequencer: Add bKeyCreated for FKeyPropertyResult and notify movie scene data changed when a key is created.

Change 3795457 by [REDACTED_PERSON]

	Reverted one of the changes from CL 3793341 as it was causing a crash at exit

Change 3795739 by [REDACTED_PERSON]

	Sequencer: Set ease in/out objects as transactional.

Change 3796871 by [REDACTED_PERSON]

	MediaAssets: Fixed media texture initialization related crash on Metal; fixed white output on inactive media player; added proper support for clear color

Change 3796921 by [REDACTED_PERSON]

	Fix for a crash in AvfMediaPlayer related to reusing FAvfMediaPlayer object

Change 3797176 by [REDACTED_PERSON]

	MediaAssets: Fixed white flicker when rapidly swapping media players on a media texture

	Also no longer clearing media textures each frame if they are not in use

Change 3797515 by [REDACTED_PERSON]

	Actor Sequence: Add nullptr check - fix actor sequence crash when actor sequence doesn't exist.

diff --git a/Engine/Build/Android/Java/src/com/epicgames/ue4/MediaPlayer14.java b/Engine/Build/Android/Java/src/com/epicgames/ue4/MediaPlayer14.java
--- a/Engine/Build/Android/Java/src/com/epicgames/ue4/MediaPlayer14.java
+++ b/Engine/Build/Android/Java/src/com/epicgames/ue4/MediaPlayer14.java
@@ -485,6 +485,21 @@ public class MediaPlayer14
 		return position;
 	}
 
+	public int getDuration()
+	{
+		int duration = 0;
+		
+		synchronized(this)
+		{
+			if (Prepared)
+			{
+				duration = super.getDuration();
+			}
+		}
+
+		return duration;
+	}
+	
 	public void seekTo(int position)
 	{
 		synchronized (this)
diff --git a/Engine/Config/BaseEngine.ini b/Engine/Config/BaseEngine.ini
--- a/Engine/Config/BaseEngine.ini
+++ b/Engine/Config/BaseEngine.ini
@@ -478,6 +478,8 @@ bUseStreamingPause=false
 +PropertyRedirects=(OldName="FFontParameterValue.ParameterName", NewName="FFontParameterValue.ParameterName_DEPRECATED")
 
 ; 4.19
++ClassRedirects=(OldName="/Script/MovieSceneTracks.MovieSceneSubTrack",NewName="/Script/MovieScene.MovieSceneSubTrack")
++ClassRedirects=(OldName="/Script/MovieSceneTracks.MovieSceneSubSection",NewName="/Script/MovieScene.MovieSceneSubSection")
 +FunctionRedirects=(OldName="InverseLerp",NewName="/Script/Engine.KismetMathLibrary.NormalizeToRange")
 +PropertyRedirects=(OldName="NormalizeToRange.A",NewName="RangeMin")
 +PropertyRedirects=(OldName="NormalizeToRange.B",NewName="RangeMax")
diff --git a/Engine/Plugins/Experimental/ControlRig/Source/ControlRig/Private/Sequencer/ControlRigBindingTemplate.cpp b/Engine/Plugins/Experimental/ControlRig/Source/ControlRig/Private/Sequencer/ControlRigBindingTemplate.cpp
--- a/Engine/Plugins/Experimental/ControlRig/Source/ControlRig/Private/Sequencer/ControlRigBindingTemplate.cpp
+++ b/Engine/Plugins/Experimental/ControlRig/Source/ControlRig/Private/Sequencer/ControlRigBindingTemplate.cpp
@@ -8,6 +8,7 @@
 #include "ControlRig.h"
 #include "Components/SkeletalMeshComponent.h"
 #include "ControlRigSequencerAnimInstance.h"
+#include "MovieSceneControlRigInstanceData.h"
 
 DECLARE_CYCLE_STAT(TEXT("Binding Track Evaluate"), MovieSceneEval_BindControlRigTemplate_Evaluate, STATGROUP_MovieSceneEval);
 DECLARE_CYCLE_STAT(TEXT("Binding Track Token Execute"), MovieSceneEval_BindControlRig_TokenExecute, STATGROUP_MovieSceneEval);
@@ -62,33 +63,20 @@ struct FControlRigPreAnimatedTokenProducer : IMovieScenePreAnimatedTokenProducer
 
 struct FBindControlRigObjectToken : IMovieSceneExecutionToken
 {
-	FBindControlRigObjectToken(FGuid InObjectBindingId, FMovieSceneSequenceIDRef InObjectBindingSequenceID, float InWeight, bool bInAdditive, bool bInApplyBoneFilter, const FInputBlendPose& InBoneFilter, bool bInSpawned)
+
+	FBindControlRigObjectToken(float InWeight, bool bInSpawned
+#if WITH_EDITORONLY_DATA
+		, TWeakObjectPtr<> InObjectBinding = TWeakObjectPtr<>()
+#endif
+		)
 		: 
 #if WITH_EDITORONLY_DATA
-		ObjectBinding(nullptr),
+		  ObjectBinding(InObjectBinding),
 #endif
-		ObjectBindingId(InObjectBindingId)
-		, ObjectBindingSequenceID(InObjectBindingSequenceID)
-		, BoneFilter(InBoneFilter)
-		, Weight(InWeight)
-		, bApplyBoneFilter(bInApplyBoneFilter)
-		, bAdditive(bInAdditive)
+		  Weight(InWeight)
 		, bSpawned(bInSpawned)
 	{}
 
-#if WITH_EDITORONLY_DATA
-	FBindControlRigObjectToken(TWeakObjectPtr<> InObjectBinding, float InWeight, bool bInAdditive, bool bInApplyBoneFilter, const FInputBlendPose& InBoneFilter, bool bInSpawned)
-		: ObjectBinding(InObjectBinding)
-		, BoneFilter(InBoneFilter)
-		, Weight(InWeight)
-		, bApplyBoneFilter(bInApplyBoneFilter)
-		, bAdditive(bInAdditive)
-		, bSpawned(bInSpawned)
-	{
-		ObjectBindingId.Invalidate();
-	}
-#endif
-
 	void BindToSequencerInstance(UControlRig* ControlRig)
 	{
 		check(ControlRig);
@@ -114,6 +102,21 @@ struct FBindControlRigObjectToken : IMovieSceneExecutionToken
 	{
 		MOVIESCENE_DETAILED_SCOPE_CYCLE_COUNTER(MovieSceneEval_BindControlRig_TokenExecute)
 
+		FInputBlendPose DefaultBoneFilter;
+
+		FMovieSceneEvaluationOperand InstanceOperand;
+		bool bAdditive = false;
+		bool bApplyBoneFilter = false;
+		const FInputBlendPose* BoneFilter = &DefaultBoneFilter;
+
+		if (const FMovieSceneControlRigInstanceData* InstanceData = PersistentData.FindInstanceData<FMovieSceneControlRigInstanceData>())
+		{
+			InstanceOperand = InstanceData->Operand;
+			bAdditive = InstanceData->bAdditive;
+			bApplyBoneFilter = InstanceData->bApplyBoneFilter;
+			BoneFilter = &InstanceData->BoneFilter;
+		}
+
 		TArrayView<TWeakObjectPtr<>> BoundObjects = Player.FindBoundObjects(Operand);
 		const bool bHasBoundObjects = BoundObjects.Num() != 0;
 		if (bSpawned)
@@ -127,9 +130,9 @@ struct FBindControlRigObjectToken : IMovieSceneExecutionToken
 				if (Sequence)
 				{
 					ControlRig = CastChecked<UControlRig>(Player.GetSpawnRegister().SpawnObject(Operand.ObjectBindingID, *Sequence->GetMovieScene(), Operand.SequenceID, Player));
-					if (ObjectBindingId.IsValid())
+					if (InstanceOperand.ObjectBindingID.IsValid())
 					{
-						TArrayView<TWeakObjectPtr<>> OuterBoundObjects = Player.FindBoundObjects(ObjectBindingId, ObjectBindingSequenceID);
+						TArrayView<TWeakObjectPtr<>> OuterBoundObjects = Player.FindBoundObjects(InstanceOperand);
 						if (OuterBoundObjects.Num() > 0)
 						{
 							UObject* OuterBoundObject = OuterBoundObjects[0].Get();
@@ -175,7 +178,7 @@ struct FBindControlRigObjectToken : IMovieSceneExecutionToken
 				{
 					if (UControlRigSequencerAnimInstance* AnimInstance = Cast<UControlRigSequencerAnimInstance>(SkeletalMeshComponent->GetAnimInstance()))
 					{
-						bool bStructureChanged = AnimInstance->UpdateControlRig(ControlRig, Operand.SequenceID.GetInternalValue(), bAdditive, bApplyBoneFilter, BoneFilter, Weight);
+						bool bStructureChanged = AnimInstance->UpdateControlRig(ControlRig, Operand.SequenceID.GetInternalValue(), bAdditive, bApplyBoneFilter, *BoneFilter, Weight);
 						if (bStructureChanged)
 						{
 							AnimInstance->RecalcRequiredBones();
@@ -204,7 +207,7 @@ struct FBindControlRigObjectToken : IMovieSceneExecutionToken
 						if (UControlRigSequencerAnimInstance* AnimInstance = Cast<UControlRigSequencerAnimInstance>(SkeletalMeshComponent->GetAnimInstance()))
 						{
 							// Force us to zero weight before we despawn, as the graph could persist
-							AnimInstance->UpdateControlRig(ControlRig, Operand.SequenceID.GetInternalValue(), bAdditive, bApplyBoneFilter, BoneFilter, 0.0f);
+							AnimInstance->UpdateControlRig(ControlRig, Operand.SequenceID.GetInternalValue(), bAdditive, bApplyBoneFilter, *BoneFilter, 0.0f);
 							AnimInstance->RecalcRequiredBones();
 						}
 						UAnimSequencerInstance::UnbindFromSkeletalMeshComponent(SkeletalMeshComponent);
@@ -223,34 +226,16 @@ struct FBindControlRigObjectToken : IMovieSceneExecutionToken
 	TWeakObjectPtr<> ObjectBinding;
 #endif
 
-	/** The object binding controllers should bind to (in the case we are bound to a sequencer object) */
-	FGuid ObjectBindingId;
-
-	/** The sequence ID controllers should bind to */
-	FMovieSceneSequenceID ObjectBindingSequenceID;
-
-	/** Per-bone filter to apply to our animation */
-	FInputBlendPose BoneFilter;
-
 	/** The weight to apply this controller at */
 	float Weight;
 
-	/** Only apply bones that are in the filter */
-	bool bApplyBoneFilter;
-
-	/** Whether we are additive */
-	bool bAdditive;
-
 	/** Whether this token should spawn an object */
 	bool bSpawned;
 };
 
 FControlRigBindingTemplate::FControlRigBindingTemplate(const UMovieSceneSpawnSection& SpawnSection)
 	: FMovieSceneSpawnSectionTemplate(SpawnSection)
-	, bApplyBoneFilter(false)
-	, bAdditive(false)
 {
-	WeightCurve.SetDefaultValue(1.0f);
 }
 
 #if WITH_EDITORONLY_DATA
@@ -272,47 +257,24 @@ void FControlRigBindingTemplate::ClearObjectBinding()
 
 #endif
 
-void FControlRigBindingTemplate::SetObjectBindingId(FGuid InObjectBindingId, FMovieSceneSequenceIDRef InObjectBindingSequenceID)
-{
-	ObjectBindingId = InObjectBindingId;
-	ObjectBindingSequenceID = InObjectBindingSequenceID;
-}
-
-void FControlRigBindingTemplate::SetWeightCurve(const FRichCurve& InWeightCurve, float InOffset, float InScale)
-{ 
-	WeightCurve = InWeightCurve;
-
-	WeightCurve.ShiftCurve(InOffset);
-	WeightCurve.ScaleCurve(0.0f, InScale);
-}
-
-void FControlRigBindingTemplate::SetPerBoneBlendFilter(bool bInApplyBoneFilter, const FInputBlendPose& InBoneFilter)
-{
-	bApplyBoneFilter = bInApplyBoneFilter;
-	if (bInApplyBoneFilter)
-	{
-		BoneFilter = InBoneFilter;
-	}
-	else
-	{
-		BoneFilter.BranchFilters.Empty();
-	}
-}
 
 void FControlRigBindingTemplate::Evaluate(const FMovieSceneEvaluationOperand& Operand, const FMovieSceneContext& Context, const FPersistentEvaluationData& PersistentData, FMovieSceneExecutionTokens& ExecutionTokens) const
 {
 	MOVIESCENE_DETAILED_SCOPE_CYCLE_COUNTER(MovieSceneEval_BindControlRigTemplate_Evaluate)
 
+	const FMovieSceneControlRigInstanceData* InstanceData = PersistentData.FindInstanceData<FMovieSceneControlRigInstanceData>();
+
 	const float Time = Context.GetTime();
+	const float Weight = InstanceData ? InstanceData->Weight.Eval(Time) : 1.f;
 
-	if (ObjectBindingId.IsValid())
+	if (InstanceData && InstanceData->Operand.ObjectBindingID.IsValid())
 	{
-		ExecutionTokens.Add(FBindControlRigObjectToken(ObjectBindingId, ObjectBindingSequenceID, WeightCurve.Eval(Time), bAdditive, bApplyBoneFilter, BoneFilter, Curve.Evaluate(Time) != 0));
+		ExecutionTokens.Add(FBindControlRigObjectToken(Weight, Curve.Evaluate(Time) != 0));
 	}
 #if WITH_EDITORONLY_DATA
 	else if (ObjectBinding.IsValid())
 	{
-		ExecutionTokens.Add(FBindControlRigObjectToken(ObjectBinding, WeightCurve.Eval(Time), bAdditive, bApplyBoneFilter, BoneFilter, Curve.Evaluate(Time) != 0));
+		ExecutionTokens.Add(FBindControlRigObjectToken(Weight, Curve.Evaluate(Time) != 0, ObjectBinding));
 	}
 #endif
 }
diff --git a/Engine/Plugins/Experimental/ControlRig/Source/ControlRig/Private/Sequencer/ControlRigSequencerAnimInstanceProxy.cpp b/Engine/Plugins/Experimental/ControlRig/Source/ControlRig/Private/Sequencer/ControlRigSequencerAnimInstanceProxy.cpp
--- a/Engine/Plugins/Experimental/ControlRig/Source/ControlRig/Private/Sequencer/ControlRigSequencerAnimInstanceProxy.cpp
+++ b/Engine/Plugins/Experimental/ControlRig/Source/ControlRig/Private/Sequencer/ControlRigSequencerAnimInstanceProxy.cpp
@@ -56,8 +56,8 @@ void FControlRigSequencerAnimInstanceProxy::ResetNodes()
 {
 	FAnimSequencerInstanceProxy::ResetNodes();
 
-	FMemory::Memzero(LayeredBoneBlendNode.BlendWeights.GetData(), AdditiveLayeredBoneBlendNode.BlendWeights.GetAllocatedSize());
-	FMemory::Memzero(LayeredBoneBlendNode.BlendWeights.GetData(), AdditiveLayeredBoneBlendNode.BlendWeights.GetAllocatedSize());
+	FMemory::Memzero(LayeredBoneBlendNode.BlendWeights.GetData(), LayeredBoneBlendNode.BlendWeights.GetAllocatedSize());
+	FMemory::Memzero(AdditiveLayeredBoneBlendNode.BlendWeights.GetData(), AdditiveLayeredBoneBlendNode.BlendWeights.GetAllocatedSize());
 }
 
 void FControlRigSequencerAnimInstanceProxy::InitControlRigTrack(UControlRig* InControlRig, bool bAdditive, bool bApplyBoneFilter, const FInputBlendPose& BoneFilter, uint32 SequenceId)
diff --git a/Engine/Plugins/Experimental/ControlRig/Source/ControlRig/Private/Sequencer/MovieSceneControlRigSection.cpp b/Engine/Plugins/Experimental/ControlRig/Source/ControlRig/Private/Sequencer/MovieSceneControlRigSection.cpp
--- a/Engine/Plugins/Experimental/ControlRig/Source/ControlRig/Private/Sequencer/MovieSceneControlRigSection.cpp
+++ b/Engine/Plugins/Experimental/ControlRig/Source/ControlRig/Private/Sequencer/MovieSceneControlRigSection.cpp
@@ -6,6 +6,7 @@
 #include "MovieScene.h"
 #include "ControlRigSequence.h"
 #include "ControlRigBindingTemplate.h"
+#include "MovieSceneControlRigInstanceData.h"
 
 #define LOCTEXT_NAMESPACE "MovieSceneControlRigSection"
 
@@ -52,37 +53,29 @@ void UMovieSceneControlRigSection::GetKeyHandles(TSet<FKeyHandle>& OutKeyHandles
 	}
 }
 
-FMovieSceneEvaluationTemplate& UMovieSceneControlRigSection::GenerateTemplateForSubSequence(const FMovieSceneTrackCompilerArgs& InArgs) const
+FMovieSceneSubSequenceData UMovieSceneControlRigSection::GenerateSubSequenceData(const FSubSequenceInstanceDataParams& Params) const
 {
-	// Use our section as the object key here
-	FMovieSceneEvaluationTemplate& Template = InArgs.SubSequenceStore.GetCompiledTemplate(*SubSequence, FObjectKey(this));
+	FMovieSceneSubSequenceData SubData(*this);
 
-	for (TPair<FMovieSceneTrackIdentifier, FMovieSceneEvaluationTrack>& TrackPair : Template.GetTracks())
+	FMovieSceneControlRigInstanceData InstanceData;
+	InstanceData.bAdditive = bAdditive;
+	InstanceData.bApplyBoneFilter = bApplyBoneFilter;
+	if (InstanceData.bApplyBoneFilter)
 	{
-		// iterate child templates
-		for (FMovieSceneEvalTemplatePtr& ChildTemplate : TrackPair.Value.GetChildTemplates())
-		{
-			if (&(*ChildTemplate).GetScriptStruct() == FControlRigBindingTemplate::StaticStruct())
-			{
-				// set section curves into binding track
-				FControlRigBindingTemplate& BindingTemplate = *static_cast<FControlRigBindingTemplate*>(&(*ChildTemplate));
-
-				BindingTemplate.SetObjectBindingId(InArgs.ObjectBindingId, MovieSceneSequenceID::Root);
-				BindingTemplate.SetWeightCurve(Weight, -GetStartTime(), Parameters.TimeScale == 0.0f ? 0.0f : 1.0f / Parameters.TimeScale);
-				BindingTemplate.SetPerBoneBlendFilter(bApplyBoneFilter, BoneFilter);
-				BindingTemplate.SetAdditive(bAdditive);
-			}
-		}
+		InstanceData.BoneFilter = BoneFilter;
 	}
+	else
+	{
+		InstanceData.BoneFilter.BranchFilters.Empty();
+	}
+	InstanceData.Operand = Params.Operand;
 
-	return Template;
-}
-
-FMovieSceneSubSequenceData UMovieSceneControlRigSection::GenerateSubSequenceData() const
-{
-	FMovieSceneSubSequenceData SubData = UMovieSceneSubSection::GenerateSubSequenceData();
+	InstanceData.Weight.SetDefaultValue(1.0f);
+	InstanceData.Weight = Weight;
+	InstanceData.Weight.ShiftCurve(-GetStartTime());
+	InstanceData.Weight.ScaleCurve(0.0f, Parameters.TimeScale == 0.0f ? 0.0f : 1.0f / Parameters.TimeScale);
 
-	SubData.SequenceKeyObject = this;
+	SubData.InstanceData = FMovieSceneSequenceInstanceDataPtr(InstanceData);
 
 	return SubData;
 }
diff --git a/Engine/Plugins/Experimental/ControlRig/Source/ControlRig/Public/Sequencer/ControlRigBindingTemplate.h b/Engine/Plugins/Experimental/ControlRig/Source/ControlRig/Public/Sequencer/ControlRigBindingTemplate.h
--- a/Engine/Plugins/Experimental/ControlRig/Source/ControlRig/Public/Sequencer/ControlRigBindingTemplate.h
+++ b/Engine/Plugins/Experimental/ControlRig/Source/ControlRig/Public/Sequencer/ControlRigBindingTemplate.h
@@ -16,30 +16,12 @@ struct CONTROLRIG_API FControlRigBindingTemplate : public FMovieSceneSpawnSectio
 {
 	GENERATED_BODY()
 	
-	FControlRigBindingTemplate()
-		: ObjectBindingSequenceID(MovieSceneSequenceID::Root)
-		, bApplyBoneFilter(false)
-		, bAdditive(false)
-	{
-		WeightCurve.SetDefaultValue(1.0f);
-	}
+	FControlRigBindingTemplate(){}
 
 	FControlRigBindingTemplate(const UMovieSceneSpawnSection& SpawnSection);
 
 	static FMovieSceneAnimTypeID GetAnimTypeID();
 
-	/** Set the object & sequence we are bound to */
-	void SetObjectBindingId(FGuid InObjectBindingId, FMovieSceneSequenceIDRef InObjectBindingSequenceID);
-
-	/** Set the weight curve for this template */
-	void SetWeightCurve(const FRichCurve& InWeightCurve, float Offset, float Scale);
-
-	/** Set whether we are additive */
-	void SetAdditive(bool bInAdditive) { bAdditive = bInAdditive; }
-
-	/** Set whether we only apply modified bones */
-	void SetPerBoneBlendFilter(bool bInApplyBoneFilter, const FInputBlendPose& InBoneFilter);
-
 	// Bind to a runtime (non-sequencer-controlled) object
 #if WITH_EDITORONLY_DATA
 	static void SetObjectBinding(TWeakObjectPtr<> InObjectBinding);
@@ -57,28 +39,4 @@ private:
 	/** The current external object binding we are using */
 	static TWeakObjectPtr<UObject> ObjectBinding;
 #endif
-
-	/** The current internal object binding we are using */
-	UPROPERTY()
-	FGuid ObjectBindingId;
-
-	/** The current sequence of the internal object binding we are using */
-	UPROPERTY()
-	FMovieSceneSequenceID ObjectBindingSequenceID;
-
-	/** Weight curve to evaluate this rig with */
-	UPROPERTY()
-	FRichCurve WeightCurve;
-
-	/** Per-bone filter to apply to our animation */
-	UPROPERTY()
-	FInputBlendPose BoneFilter;
-
-	/** Only apply bones that are in the filter */
-	UPROPERTY()
-	bool bApplyBoneFilter;
-
-	/** Whether we are additive */
-	UPROPERTY()
-	bool bAdditive;
 };
diff --git a/Engine/Plugins/Experimental/ControlRig/Source/ControlRig/Public/Sequencer/MovieSceneControlRigSection.h b/Engine/Plugins/Experimental/ControlRig/Source/ControlRig/Public/Sequencer/MovieSceneControlRigSection.h
--- a/Engine/Plugins/Experimental/ControlRig/Source/ControlRig/Public/Sequencer/MovieSceneControlRigSection.h
+++ b/Engine/Plugins/Experimental/ControlRig/Source/ControlRig/Public/Sequencer/MovieSceneControlRigSection.h
@@ -49,7 +49,5 @@ public:
 	virtual void GetKeyHandles(TSet<FKeyHandle>& OutKeyHandles, TRange<float> TimeRange) const override;
 
 	//~ UMovieSceneSubSection interface
-
-	virtual FMovieSceneEvaluationTemplate& GenerateTemplateForSubSequence(const FMovieSceneTrackCompilerArgs& InArgs) const override;
-	virtual FMovieSceneSubSequenceData GenerateSubSequenceData() const override;
+	virtual FMovieSceneSubSequenceData GenerateSubSequenceData(const FSubSequenceInstanceDataParams& Params) const override;
 };
diff --git a/Engine/Plugins/Experimental/ControlRig/Source/ControlRigEditor/Private/ControlRigEditorModule.cpp b/Engine/Plugins/Experimental/ControlRig/Source/ControlRigEditor/Private/ControlRigEditorModule.cpp
--- a/Engine/Plugins/Experimental/ControlRig/Source/ControlRigEditor/Private/ControlRigEditorModule.cpp
+++ b/Engine/Plugins/Experimental/ControlRig/Source/ControlRigEditor/Private/ControlRigEditorModule.cpp
@@ -464,6 +464,8 @@ void FControlRigEditorModule::HandleSequencerCreated(TSharedRef<ISequencer> InSe
 			}
 		}
 	});
+
+	InSequencer->OnGetIsTrackVisible().BindRaw(this, &FControlRigEditorModule::IsTrackVisible);
 }
 
 void FControlRigEditorModule::HandleAssetEditorOpened(UObject* InAsset)
@@ -609,6 +611,21 @@ void FControlRigEditorModule::ReImportFromRigSequence(TArray<FAssetData> InAsset
 	}
 }
 
+bool FControlRigEditorModule::IsTrackVisible(const UMovieSceneTrack* InTrack)
+{
+	if (FControlRigEditMode* ControlRigEditMode = static_cast<FControlRigEditMode*>(GLevelEditorModeTools().GetActiveMode(FControlRigEditMode::ModeName)))
+	{		
+		// If nothing selected, show all nodes
+		if (ControlRigEditMode->GetSelectedNodes().Num() == 0)
+		{
+			return true;
+		}
+
+		return ControlRigEditMode->IsNodeSelected(ControlRigEditMode->GetNodeFromPropertyPath(InTrack->GetTrackName().ToString()));
+	}
+	return true;
+}
+
 IMPLEMENT_MODULE(FControlRigEditorModule, ControlRigEditor)
 
 #undef LOCTEXT_NAMESPACE
diff --git a/Engine/Plugins/Experimental/ControlRig/Source/ControlRigEditor/Private/ControlRigEditorModule.h b/Engine/Plugins/Experimental/ControlRig/Source/ControlRigEditor/Private/ControlRigEditorModule.h
--- a/Engine/Plugins/Experimental/ControlRig/Source/ControlRigEditor/Private/ControlRigEditorModule.h
+++ b/Engine/Plugins/Experimental/ControlRig/Source/ControlRigEditor/Private/ControlRigEditorModule.h
@@ -16,6 +16,7 @@ class USkeletalMesh;
 class FToolBarBuilder;
 class FExtender;
 class FUICommandList;
+class UMovieSceneTrack;
 
 class FControlRigEditorModule : public IModuleInterface
 {
@@ -60,6 +61,9 @@ private:
 	/** Bind our module-level commands */
 	void BindCommands();
 
+	/** Whether the track is visible in the sequencer node tree */
+	bool IsTrackVisible(const UMovieSceneTrack* InTrack);
+
 	/** Compiler customization for animation controllers */
 	FControlRigBlueprintCompiler ControlRigBlueprintCompiler;
 
diff --git a/Engine/Plugins/Experimental/ControlRig/Source/ControlRigEditor/Private/EditMode/ControlRigEditMode.cpp b/Engine/Plugins/Experimental/ControlRig/Source/ControlRigEditor/Private/EditMode/ControlRigEditMode.cpp
--- a/Engine/Plugins/Experimental/ControlRig/Source/ControlRigEditor/Private/EditMode/ControlRigEditMode.cpp
+++ b/Engine/Plugins/Experimental/ControlRig/Source/ControlRigEditor/Private/EditMode/ControlRigEditMode.cpp
@@ -11,6 +11,7 @@
 #include "HumanRig.h"
 #include "ControlRigEditModeSettings.h"
 #include "ISequencer.h"
+#include "SequencerSettings.h"
 #include "ControlRigSequence.h"
 #include "ControlRigBindingTemplate.h"
 #include "Sections/MovieSceneSpawnSection.h"
@@ -1014,6 +1015,25 @@ void FControlRigEditMode::SetNodeSelectionByPropertyPath(const TArray<FString>&
 	}
 }
 
+FName FControlRigEditMode::GetNodeFromPropertyPath(const FString& InPropertyPath) const
+{
+	for (TWeakObjectPtr<UControlRig> ControlRig : ControlRigs)
+	{
+		if (UHierarchicalRig* HierarchicalRig = Cast<UHierarchicalRig>(ControlRig.Get()))
+		{
+			for (UControlManipulator* Manipulator : HierarchicalRig->Manipulators)
+			{
+				if (InPropertyPath == Manipulator->PropertyToManipulate.ToString())
+				{
+					return Manipulator->Name;
+				}
+			}
+		}
+	}
+	
+	return NAME_None;
+}
+
 void FControlRigEditMode::HandleObjectSpawned(FGuid InObjectBinding, UObject* SpawnedObject, IMovieScenePlayer& Player)
 {
 	if (WeakSequencer.IsValid())
@@ -1175,6 +1195,14 @@ void FControlRigEditMode::HandleSelectionChanged(const TArray<FName>& InSelected
 	{
 		TrajectoryCache.RebuildMesh(SelectedIndices);
 	}
+
+	if (WeakSequencer.IsValid())
+	{
+		if (WeakSequencer.Pin()->GetSequencerSettings()->GetShowSelectedNodesOnly())
+		{
+			WeakSequencer.Pin()->RefreshTree();
+		}
+	}
 }
 
 void FControlRigEditMode::BindCommands()
diff --git a/Engine/Plugins/Experimental/ControlRig/Source/ControlRigEditor/Private/EditMode/ControlRigEditMode.h b/Engine/Plugins/Experimental/ControlRig/Source/ControlRigEditor/Private/EditMode/ControlRigEditMode.h
--- a/Engine/Plugins/Experimental/ControlRig/Source/ControlRigEditor/Private/EditMode/ControlRigEditMode.h
+++ b/Engine/Plugins/Experimental/ControlRig/Source/ControlRigEditor/Private/EditMode/ControlRigEditMode.h
@@ -77,6 +77,9 @@ public:
 	/** Attempt to select by property path */
 	void SetNodeSelectionByPropertyPath(const TArray<FString>& InPropertyPaths);
 
+	/** Get the node name from the property path */
+	FName GetNodeFromPropertyPath(const FString& PropertyPath) const;
+
 	/** 
 	 * Lets the edit mode know that an object has just been spawned. 
 	 * Allows us to redisplay different underlying objects in the details panel.
diff --git a/Engine/Plugins/Experimental/ControlRig/Source/ControlRigEditor/Private/Sequencer/ControlRigEditorObjectSpawner.cpp b/Engine/Plugins/Experimental/ControlRig/Source/ControlRigEditor/Private/Sequencer/ControlRigEditorObjectSpawner.cpp
--- a/Engine/Plugins/Experimental/ControlRig/Source/ControlRigEditor/Private/Sequencer/ControlRigEditorObjectSpawner.cpp
+++ b/Engine/Plugins/Experimental/ControlRig/Source/ControlRigEditor/Private/Sequencer/ControlRigEditorObjectSpawner.cpp
@@ -35,7 +35,7 @@ UObject* FControlRigEditorObjectSpawner::SpawnObject(FMovieSceneSpawnable& Spawn
 
 #if WITH_EDITOR
 
-TValueOrError<FNewSpawnable, FText> FControlRigEditorObjectSpawner::CreateNewSpawnableType(UObject& SourceObject, UMovieScene& OwnerMovieScene)
+TValueOrError<FNewSpawnable, FText> FControlRigEditorObjectSpawner::CreateNewSpawnableType(UObject& SourceObject, UMovieScene& OwnerMovieScene, UActorFactory* ActorFactory)
 {
 	// Right now we only support creating a spawnable for classes
 	if (UClass* InClass = Cast<UClass>(&SourceObject))
diff --git a/Engine/Plugins/Experimental/ControlRig/Source/ControlRigEditor/Private/Sequencer/ControlRigEditorObjectSpawner.h b/Engine/Plugins/Experimental/ControlRig/Source/ControlRigEditor/Private/Sequencer/ControlRigEditorObjectSpawner.h
--- a/Engine/Plugins/Experimental/ControlRig/Source/ControlRigEditor/Private/Sequencer/ControlRigEditorObjectSpawner.h
+++ b/Engine/Plugins/Experimental/ControlRig/Source/ControlRigEditor/Private/Sequencer/ControlRigEditorObjectSpawner.h
@@ -14,7 +14,7 @@ public:
 	virtual bool IsEditor() const override { return true; }
 	virtual UObject* SpawnObject(FMovieSceneSpawnable& Spawnable, FMovieSceneSequenceIDRef TemplateID, IMovieScenePlayer& Player) override;
 #if WITH_EDITOR
-	virtual TValueOrError<FNewSpawnable, FText> CreateNewSpawnableType(UObject& SourceObject, UMovieScene& OwnerMovieScene) override;
+	virtual TValueOrError<FNewSpawnable, FText> CreateNewSpawnableType(UObject& SourceObject, UMovieScene& OwnerMovieScene, UActorFactory* ActorFactory = nullptr) override;
 	virtual void SetupDefaultsForSpawnable(UObject* SpawnedObject, const FGuid& Guid, const FTransformData& TransformData, TSharedRef<ISequencer> Sequencer, USequencerSettings* Settings) override;
 	virtual bool CanConvertSpawnableToPossessable(FMovieSceneSpawnable& Spawnable) const override { return false; }
 #endif
diff --git a/Engine/Plugins/Experimental/ControlRig/Source/ControlRigEditor/Private/Sequencer/ControlRigTrackEditor.cpp b/Engine/Plugins/Experimental/ControlRig/Source/ControlRigEditor/Private/Sequencer/ControlRigTrackEditor.cpp
--- a/Engine/Plugins/Experimental/ControlRig/Source/ControlRigEditor/Private/Sequencer/ControlRigTrackEditor.cpp
+++ b/Engine/Plugins/Experimental/ControlRig/Source/ControlRigEditor/Private/Sequencer/ControlRigTrackEditor.cpp
@@ -209,6 +209,11 @@ TSharedRef<ISequencerTrackEditor> FControlRigTrackEditor::CreateTrackEditor( TSh
 	return MakeShareable( new FControlRigTrackEditor( InSequencer ) );
 }
 
+bool FControlRigTrackEditor::SupportsSequence(UMovieSceneSequence* InSequence) const
+{
+	return (InSequence != nullptr) && (InSequence->GetClass()->GetName() == TEXT("ControlRigSequence"));
+}
+
 bool FControlRigTrackEditor::SupportsType( TSubclassOf<UMovieSceneTrack> Type ) const
 {
 	return Type == UMovieSceneControlRigTrack::StaticClass();
diff --git a/Engine/Plugins/Experimental/ImagePlate/Source/ImagePlate/Private/ImagePlateComponent.cpp b/Engine/Plugins/Experimental/ImagePlate/Source/ImagePlate/Private/ImagePlateComponent.cpp
--- a/Engine/Plugins/Experimental/ImagePlate/Source/ImagePlate/Private/ImagePlateComponent.cpp
+++ b/Engine/Plugins/Experimental/ImagePlate/Source/ImagePlate/Private/ImagePlateComponent.cpp
@@ -195,6 +195,7 @@ namespace
 				Mesh.DepthPriorityGroup      = SDPG_World;
 				Mesh.Type                    = PT_TriangleList;
 				Mesh.bDisableBackfaceCulling = true;
+				Mesh.LODIndex                = 0;
 
 				FMeshBatchElement& BatchElement = Mesh.Elements[0];
 				BatchElement.IndexBuffer = &IndexBuffer;
@@ -501,13 +502,7 @@ void UImagePlateComponent::GetProjectionMatricesFromViewTarget(AActor* InViewTar
 
 FBoxSphereBounds UImagePlateComponent::CalcBounds(const FTransform& LocalToWorld) const
 {
-	FBox MaxBox({
-		FVector(0,  1,  1),
-		FVector(0, -1,  1),
-		FVector(0,  1, -1),
-		FVector(0, -1, -1) });
-	MaxBox = MaxBox.TransformBy(LocalToWorld);
-	return FBoxSphereBounds(MaxBox);
+	return FBoxSphereBounds(FVector(0,0,0), FVector(1,  1,  1), 1.73205f).TransformBy(LocalToWorld);
 }
 
 AActor* UImagePlateComponent::FindViewTarget() const
diff --git a/Engine/Plugins/Media/AvfMedia/AvfMedia.uplugin b/Engine/Plugins/Media/AvfMedia/AvfMedia.uplugin
--- a/Engine/Plugins/Media/AvfMedia/AvfMedia.uplugin
+++ b/Engine/Plugins/Media/AvfMedia/AvfMedia.uplugin
@@ -20,7 +20,7 @@
 			"Name" : "AvfMedia",
 			"Type" : "RuntimeNoCommandlet",
 			"LoadingPhase" : "PreLoadingScreen",
-			"WhitelistPlatforms" : [ "IOS", "Mac" ]
+			"WhitelistPlatforms" : [ "IOS", "Mac", "TVOS" ]
 		},
 		{
 			"Name" : "AvfMediaEditor",
@@ -36,7 +36,7 @@
 			"Name" : "AvfMediaFactory",
 			"Type" : "RuntimeNoCommandlet",
 			"LoadingPhase" : "PostEngineInit",
-			"WhitelistPlatforms" : [ "IOS", "Mac", "Linux", "Win32", "Win64" ]
+			"WhitelistPlatforms" : [ "IOS", "Linux", "Mac", "TVOS", "Win32", "Win64" ]
 		}
 	]
 }
\ No newline at end of file
diff --git a/Engine/Plugins/Media/AvfMedia/Source/AvfMedia/Private/Player/AvfMediaPlayer.cpp b/Engine/Plugins/Media/AvfMedia/Source/AvfMedia/Private/Player/AvfMediaPlayer.cpp
--- a/Engine/Plugins/Media/AvfMedia/Source/AvfMedia/Private/Player/AvfMediaPlayer.cpp
+++ b/Engine/Plugins/Media/AvfMedia/Source/AvfMedia/Private/Player/AvfMediaPlayer.cpp
@@ -265,9 +265,9 @@ void FAvfMediaPlayer::Close()
 		if (MediaHelper != nil)
 		{
 			[[NSNotificationCenter defaultCenter] removeObserver:MediaHelper name:AVPlayerItemDidPlayToEndTimeNotification object:PlayerItem];
+			[PlayerItem removeObserver:MediaHelper forKeyPath:@"status"];
 		}
 
-		[PlayerItem removeObserver:MediaHelper forKeyPath:@"status"];
 		[PlayerItem release];
 		PlayerItem = nil;
 	}
@@ -296,7 +296,7 @@ void FAvfMediaPlayer::Close()
 	EventSink.ReceiveMediaEvent(EMediaEvent::MediaClosed);
 		
 	bPrerolled = false;
-	
+
 	CurrentRate = 0.f;
 }
 
@@ -434,16 +434,11 @@ bool FAvfMediaPlayer::Open(const FString& Url, const IMediaOptions* /*Options*/)
 
 		if ([[PlayerItem asset] statusOfValueForKey:@"tracks" error : &Error] == AVKeyValueStatusLoaded)
 		{
-			[[NSNotificationCenter defaultCenter] addObserver:MediaHelper selector:@selector(playerItemPlaybackEndReached:) name:AVPlayerItemDidPlayToEndTimeNotification object:PlayerItem];
-			
 			// File movies will be ready now
 			if (PlayerItem.status == AVPlayerItemStatusReadyToPlay)
 			{
 				OnStatusNotification(PlayerItem.status);
 			}
-			
-			// Streamed movies might not be and we want to know if it ever fails.
-			[PlayerItem addObserver:MediaHelper forKeyPath:@"status" options:0 context:PlayerItem];
 		}
 		else if (Error != nullptr)
 		{
@@ -460,6 +455,9 @@ bool FAvfMediaPlayer::Open(const FString& Url, const IMediaOptions* /*Options*/)
 		}
 	}];
 
+	[[NSNotificationCenter defaultCenter] addObserver:MediaHelper selector:@selector(playerItemPlaybackEndReached:) name:AVPlayerItemDidPlayToEndTimeNotification object:PlayerItem];
+	[PlayerItem addObserver:MediaHelper forKeyPath:@"status" options:0 context:PlayerItem];
+
 	[MediaPlayer replaceCurrentItemWithPlayerItem : PlayerItem];
 	[[MediaPlayer currentItem] seekToTime:kCMTimeZero];
 
diff --git a/Engine/Plugins/Media/AvfMedia/Source/AvfMedia/Private/Player/AvfMediaTracks.cpp b/Engine/Plugins/Media/AvfMedia/Source/AvfMedia/Private/Player/AvfMediaTracks.cpp
--- a/Engine/Plugins/Media/AvfMedia/Source/AvfMedia/Private/Player/AvfMediaTracks.cpp
+++ b/Engine/Plugins/Media/AvfMedia/Source/AvfMedia/Private/Player/AvfMediaTracks.cpp
@@ -66,10 +66,12 @@ struct AudioTrackTapContextData
 	AudioStreamBasicDescription DestinationFormat;
 	FMediaSamples&				SampleQueue;
 	FAvfMediaAudioSamplePool*	AudioSamplePool;
+	bool						bActive;
 	
 	AudioTrackTapContextData(FMediaSamples& InSampleQueue, FAvfMediaAudioSamplePool* InAudioSamplePool, AudioStreamBasicDescription const & InDestinationFormat)
 	: SampleQueue(InSampleQueue)
 	, AudioSamplePool(InAudioSamplePool)
+	, bActive(false)
 	{
 		FMemory::Memcpy(&DestinationFormat, &InDestinationFormat, sizeof(AudioStreamBasicDescription));
 	}
@@ -87,6 +89,7 @@ static void AudioTrackTapPrepare(MTAudioProcessingTapRef __nonnull TapRef, CMIte
 	if(ContextData)
 	{
 		FMemory::Memcpy(&ContextData->ProcessingFormat,ProcessingFormat,sizeof(AudioStreamBasicDescription));
+		ContextData->bActive = true;
 	}
 }
 
@@ -112,76 +115,79 @@ static void AudioTrackTapProcess(	MTAudioProcessingTapRef __nonnull TapRef,
 		// in the public interface to AvfMediaTracks which seems wrong - plus we save the extra function call in time critical code!
 		check(Ctx);
 		
-		// Compute required buffer size
-		uint32 BufferSize = (NumberFrames * ((Ctx->DestinationFormat.mBitsPerChannel / 8))) * Ctx->DestinationFormat.mChannelsPerFrame;
-		
-		//setup reasonable defaults as time range can be invalid - especially at the start of the audio track
-		FTimespan StartTime(0);
-		FTimespan Duration(((int64)NumberFrames * ETimespan::TicksPerSecond) / (int64)Ctx->DestinationFormat.mSampleRate);
-		
-		// If valid set time stamps give by the system
-		if((TimeRange.start.flags & kCMTimeFlags_Valid) != 0)
-		{
-			StartTime = (TimeRange.start.value * ETimespan::TicksPerSecond) / TimeRange.start.timescale;
-		}
-		
-		// Probably don't need to recompute duration as it's always been exactly the same as pre-computed above
-		if((TimeRange.duration.flags & kCMTimeFlags_Valid) != 0)
+			if(Ctx->bActive)
 		{
-			Duration = (TimeRange.duration.value * ETimespan::TicksPerSecond) / TimeRange.duration.timescale;
-		}
-		
-		// Get a media audio sample buffer from the pool
-		const TSharedRef<FAvfMediaAudioSample, ESPMode::ThreadSafe> AudioSample = Ctx->AudioSamplePool->AcquireShared();
-		if (AudioSample->Initialize(BufferSize,
-									NumberFrames,
-									Ctx->DestinationFormat.mChannelsPerFrame,
-									Ctx->DestinationFormat.mSampleRate,
-									StartTime,
-									Duration))
-		{
-			if((Ctx->ProcessingFormat.mFormatFlags & kAudioFormatFlagIsFloat) != 0)
+			// Compute required buffer size
+			uint32 BufferSize = (NumberFrames * ((Ctx->DestinationFormat.mBitsPerChannel / 8))) * Ctx->DestinationFormat.mChannelsPerFrame;
+			
+			//setup reasonable defaults as time range can be invalid - especially at the start of the audio track
+			FTimespan StartTime(0);
+			FTimespan Duration(((int64)NumberFrames * ETimespan::TicksPerSecond) / (int64)Ctx->DestinationFormat.mSampleRate);
+			
+			// If valid set time stamps give by the system
+			if((TimeRange.start.flags & kCMTimeFlags_Valid) != 0)
 			{
-				float* DestBuffer = (float*)AudioSample->GetMutableBuffer();
-				const uint32 BufferCount = BufferListInOut->mNumberBuffers;
-				
-				// We need to have the same amount of buffers as the channel count
-				check(BufferCount == Ctx->DestinationFormat.mChannelsPerFrame);
-				
-				// Interleave the seperate channel buffers into one buffer
-				for(uint32 b = 0; b < BufferCount;++b)
+				StartTime = (TimeRange.start.value * ETimespan::TicksPerSecond) / TimeRange.start.timescale;
+			}
+			
+			// Probably don't need to recompute duration as it's always been exactly the same as pre-computed above
+			if((TimeRange.duration.flags & kCMTimeFlags_Valid) != 0)
+			{
+				Duration = (TimeRange.duration.value * ETimespan::TicksPerSecond) / TimeRange.duration.timescale;
+			}
+			
+			// Get a media audio sample buffer from the pool
+			const TSharedRef<FAvfMediaAudioSample, ESPMode::ThreadSafe> AudioSample = Ctx->AudioSamplePool->AcquireShared();
+			if (AudioSample->Initialize(BufferSize,
+										NumberFrames,
+										Ctx->DestinationFormat.mChannelsPerFrame,
+										Ctx->DestinationFormat.mSampleRate,
+										StartTime,
+										Duration))
+			{
+				if((Ctx->ProcessingFormat.mFormatFlags & kAudioFormatFlagIsFloat) != 0)
 				{
-					AudioBuffer& Buffer = BufferListInOut->mBuffers[b];
-					
-					// We don't handle source processing interleaved formats - if this number equals mChannelPerFrame then we could just blit the data across in one go
-					check(Buffer.mNumberChannels == 1);
+					float* DestBuffer = (float*)AudioSample->GetMutableBuffer();
+					const uint32 BufferCount = BufferListInOut->mNumberBuffers;
 					
-					// Make sure each channel buffer has the right about of data for the number of frames and processing format
-					check(Buffer.mDataByteSize == NumberFrames * (Ctx->DestinationFormat.mBitsPerChannel / 8));
+					// We need to have the same amount of buffers as the channel count
+					check(BufferCount == Ctx->DestinationFormat.mChannelsPerFrame);
 					
-					float* SrcBuffer = (float*)Buffer.mData;
-					
-					// Perform interleave copy
-					for(uint32 f = 0;f < NumberFrames;++f)
+					// Interleave the seperate channel buffers into one buffer
+					for(uint32 b = 0; b < BufferCount;++b)
 					{
-						uint32 idx = b + (f * BufferCount);
-						DestBuffer[idx] = SrcBuffer[f];
+						AudioBuffer& Buffer = BufferListInOut->mBuffers[b];
+						
+						// We don't handle source processing interleaved formats - if this number equals mChannelPerFrame then we could just blit the data across in one go
+						check(Buffer.mNumberChannels == 1);
+						
+						// Make sure each channel buffer has the right about of data for the number of frames and processing format
+						check(Buffer.mDataByteSize == NumberFrames * (Ctx->DestinationFormat.mBitsPerChannel / 8));
+						
+						float* SrcBuffer = (float*)Buffer.mData;
+						
+						// Perform interleave copy
+						for(uint32 f = 0;f < NumberFrames;++f)
+						{
+							uint32 idx = b + (f * BufferCount);
+							DestBuffer[idx] = SrcBuffer[f];
+						}
+						
+						// Done with this source buffer clear it - otherwise AVPlayer will also play this audio -  we could seet a volume of 0 on the AudioMixInputParameters
+						// But that could be dangerous as OS may optimise out somethings at runtime for 0 volume tracks in the future
+						FMemory::Memset(Buffer.mData, 0, Buffer.mDataByteSize);
 					}
-					
-					// Done with this source buffer clear it - otherwise AVPlayer will also play this audio -  we could seet a volume of 0 on the AudioMixInputParameters
-					// But that could be dangerous as OS may optimise out somethings at runtime for 0 volume tracks in the future
-					FMemory::Memset(Buffer.mData, 0, Buffer.mDataByteSize);
 				}
+				else
+				{
+					// Processing format should always be float however...
+					// If we encounter this case (kAudioFormatFlagIsSignedInteger) we need to use sint16 audio sample type on FAvfMediaAudioSample
+					// i.e. make FAvfMediaAudioSample sample type settable, the engine should convert to float internally as that is it's preferred format now
+					check(false);
+				}
+				
+				Ctx->SampleQueue.AddAudio(AudioSample);
 			}
-			else
-			{
-				// Processing format should always be float however...
-				// If we encounter this case (kAudioFormatFlagIsSignedInteger) we need to use sint16 audio sample type on FAvfMediaAudioSample
-				// i.e. make FAvfMediaAudioSample sample type settable, the engine should convert to float internally as that is it's preferred format now
-				check(false);
-			}
-			
-			Ctx->SampleQueue.AddAudio(AudioSample);
 		}
 	}
 }
@@ -203,22 +209,30 @@ static void AudioTrackTapFinalize(MTAudioProcessingTapRef __nonnull TapRef)
 
 static void AudioTrackTapShutdownCurrentAudioTrackProcessing(AVPlayerItem* PlayerItem)
 {
+	SCOPED_AUTORELEASE_POOL;
+
 	check(PlayerItem);
 	
 	if(PlayerItem.audioMix != nil)
 	{
 		MTAudioProcessingTapRef TapNodeRef = ((AVMutableAudioMixInputParameters*)PlayerItem.audioMix.inputParameters[0]).audioTapProcessor;
+		
 		PlayerItem.audioMix = nil;
 		
 		if(TapNodeRef)
 		{
-			CFRelease(TapNodeRef);
+			AudioTrackTapContextData* Ctx = (AudioTrackTapContextData*) MTAudioProcessingTapGetStorage(TapNodeRef);
+			check(Ctx);
+			
+			Ctx->bActive = false;
 		}
 	}
 }
 
 static void AudioTrackTapInitializeForAudioTrack(FMediaSamples& InSampleQueue, FAvfMediaAudioSamplePool* InAudioSamplePool, AudioStreamBasicDescription const & InDestinationFormat, AVPlayerItem* PlayerItem, AVAssetTrack* AssetTrack)
 {
+	SCOPED_AUTORELEASE_POOL;
+
 	check(InAudioSamplePool);
 	check(PlayerItem);
 	
@@ -251,6 +265,8 @@ static void AudioTrackTapInitializeForAudioTrack(FMediaSamples& InSampleQueue, F
 			AudioMix.inputParameters = @[InputParams];
 			
 			PlayerItem.audioMix = AudioMix;
+
+			CFRelease(TapNodeRef);
 		}
 	}
 }
diff --git a/Engine/Plugins/Media/ImgMedia/Source/ImgMedia/Private/Player/ImgMediaPlayer.cpp b/Engine/Plugins/Media/ImgMedia/Source/ImgMedia/Private/Player/ImgMediaPlayer.cpp
--- a/Engine/Plugins/Media/ImgMedia/Source/ImgMedia/Private/Player/ImgMediaPlayer.cpp
+++ b/Engine/Plugins/Media/ImgMedia/Source/ImgMedia/Private/Player/ImgMediaPlayer.cpp
@@ -395,6 +395,11 @@ bool FImgMediaPlayer::Seek(const FTimespan& Time)
 	CurrentTime = Time;
 	LastFetchTime = FTimespan::MinValue();
 
+	if (CurrentState == EMediaState::Paused)
+	{
+		Loader->RequestFrame(CurrentTime, CurrentRate);
+	}
+
 	EventSink.ReceiveMediaEvent(EMediaEvent::SeekCompleted);
 
 	return true;
diff --git a/Engine/Plugins/Media/MediaCompositing/Source/MediaCompositing/Private/Components/MediaPlaneComponent.cpp b/Engine/Plugins/Media/MediaCompositing/Source/MediaCompositing/Private/Components/MediaPlaneComponent.cpp
--- a/Engine/Plugins/Media/MediaCompositing/Source/MediaCompositing/Private/Components/MediaPlaneComponent.cpp
+++ b/Engine/Plugins/Media/MediaCompositing/Source/MediaCompositing/Private/Components/MediaPlaneComponent.cpp
@@ -30,6 +30,8 @@
 #include "MediaPlaneFrustumComponent.h"
 #include "MediaPlaneComponent.h"
 #include "StaticMeshResources.h"
+#include "MediaTexture.h"
+#include "MediaPlayer.h"
 
 namespace
 {
@@ -300,6 +302,7 @@ FMediaPlaneParameters::FMediaPlaneParameters()
 	, FixedSize(100,100)
 	, RenderTexture(nullptr)
 	, DynamicMaterial(nullptr)
+	, MediaTexture(nullptr)
 {
 }
 
@@ -359,7 +362,7 @@ void UMediaPlaneComponent::SetMediaPlane(FMediaPlaneParameters NewPlane)
 	UpdateMaterialParametersForMedia();
 }
 
-void UMediaPlaneComponent::OnRenderTextureChanged()
+void UMediaPlaneComponent::OnMediaTextureChanged()
 {
 	UpdateMaterialParametersForMedia();
 }
@@ -397,9 +400,16 @@ void UMediaPlaneComponent::UpdateTransformScale()
 	}
 }
 
+void UMediaPlaneComponent::SetMediaTexture(UTexture* Texture)
+{
+	Plane.MediaTexture = Texture;
+
+	UpdateMaterialParametersForMedia();
+}
+
 void UMediaPlaneComponent::UpdateMaterialParametersForMedia()
 {
-	if (!Plane.TextureParameterName.IsNone() && Plane.Material && Plane.RenderTexture)
+	if (!Plane.TextureParameterName.IsNone() && Plane.Material && Plane.MediaTexture)
 	{
 		if (!Plane.DynamicMaterial)
 		{
@@ -407,7 +417,7 @@ void UMediaPlaneComponent::UpdateMaterialParametersForMedia()
 			Plane.DynamicMaterial->SetFlags(RF_Transient);
 		}
 
-		Plane.DynamicMaterial->SetTextureParameterValue(Plane.TextureParameterName, Plane.RenderTexture);
+		Plane.DynamicMaterial->SetTextureParameterValue(Plane.TextureParameterName, Plane.MediaTexture);
 	}
 	else
 	{
diff --git a/Engine/Plugins/Media/MediaCompositing/Source/MediaCompositing/Private/MovieScene/MovieSceneMediaTemplate.cpp b/Engine/Plugins/Media/MediaCompositing/Source/MediaCompositing/Private/MovieScene/MovieSceneMediaTemplate.cpp
--- a/Engine/Plugins/Media/MediaCompositing/Source/MediaCompositing/Private/MovieScene/MovieSceneMediaTemplate.cpp
+++ b/Engine/Plugins/Media/MediaCompositing/Source/MediaCompositing/Private/MovieScene/MovieSceneMediaTemplate.cpp
@@ -2,9 +2,14 @@
 
 #include "MovieSceneMediaTemplate.h"
 
+#include "MediaPlaneComponent.h"
 #include "MediaSource.h"
 #include "MediaPlayer.h"
 #include "MediaTexture.h"
+#include "Materials/MaterialInterface.h"
+#include "Materials/Material.h"
+#include "Materials/MaterialExpression.h"
+#include "Materials/MaterialExpressionTextureBase.h"
 #include "Package.h"
 #include "UObject/GCObject.h"
 
@@ -82,25 +87,25 @@ private:
 
 
 // This pre-animated token is only created if we create a new mediaplayer and assign it to the property
-struct FMediaPlayerPropertyPreAnimatedToken
+struct FMediaTexturePropertyPreAnimatedToken
 	: IMovieScenePreAnimatedToken
 {
-	FMediaPlayerPropertyPreAnimatedToken(UMediaPlayer* InPreviousPropertyValue, const TSharedPtr<FTrackInstancePropertyBindings>& InBindings)
+	FMediaTexturePropertyPreAnimatedToken(UMediaTexture* InPreviousPropertyValue, const TSharedPtr<FTrackInstancePropertyBindings>& InBindings)
 		: PreviousPropertyValue(InPreviousPropertyValue), PropertyBindings(InBindings)
 	{ }
 
 	/** Manual defaulted functions while we need to support compilers that can't generate/default them for us */
 #if PLATFORM_COMPILER_HAS_DEFAULTED_FUNCTIONS
-	FMediaPlayerPropertyPreAnimatedToken(FMediaPlayerPropertyPreAnimatedToken&& In) = default;
-	FMediaPlayerPropertyPreAnimatedToken& operator=(FMediaPlayerPropertyPreAnimatedToken&& In) = default;
+	FMediaTexturePropertyPreAnimatedToken(FMediaTexturePropertyPreAnimatedToken&& In) = default;
+	FMediaTexturePropertyPreAnimatedToken& operator=(FMediaTexturePropertyPreAnimatedToken&& In) = default;
 
 #else
-	FMediaPlayerPropertyPreAnimatedToken(FMediaPlayerPropertyPreAnimatedToken&& In)
+	FMediaTexturePropertyPreAnimatedToken(FMediaTexturePropertyPreAnimatedToken&& In)
 		: PreviousPropertyValue(MoveTemp(In.PreviousPropertyValue))
 		, PropertyBindings(MoveTemp(In.PropertyBindings))
 	{ }
 
-	FMediaPlayerPropertyPreAnimatedToken& operator=(FMediaPlayerPropertyPreAnimatedToken&& In)
+	FMediaTexturePropertyPreAnimatedToken& operator=(FMediaTexturePropertyPreAnimatedToken&& In)
 	{
 		PreviousPropertyValue = MoveTemp(In.PreviousPropertyValue);
 		PropertyBindings = MoveTemp(In.PropertyBindings);
@@ -110,20 +115,20 @@ struct FMediaPlayerPropertyPreAnimatedToken
 #endif
 
 	/** Virtual destructor. */
-	virtual ~FMediaPlayerPropertyPreAnimatedToken() { }
+	virtual ~FMediaTexturePropertyPreAnimatedToken() { }
 
 public:
 
 	virtual void RestoreState(UObject& RestoreObject, IMovieScenePlayer& Player)
 	{
-		UMediaPlayer* PreviousPlayer = PreviousPropertyValue.GetObject();
-		PropertyBindings->CallFunction<UMediaPlayer*>(RestoreObject, PreviousPlayer);
+		UMediaTexture* PreviousTexture = PreviousPropertyValue.GetObject();
+		PropertyBindings->CallFunction<UMediaTexture*>(RestoreObject, PreviousTexture);
 	}
 
 private:
 
 	/** Keep the previous player alive so we can definitely restore it later on */
-	TScopedRootObject<UMediaPlayer> PreviousPropertyValue;
+	TScopedRootObject<UMediaTexture> PreviousPropertyValue;
 
 	/** Property bindings that allow us to set the property when we've finished evaluating */
 	TSharedPtr<FTrackInstancePropertyBindings> PropertyBindings;
@@ -142,7 +147,7 @@ public:
 
 		virtual IMovieScenePreAnimatedTokenPtr CacheExistingState(UObject& Object) const override
 		{
-			return FMediaPlayerPropertyPreAnimatedToken(Bindings->GetCurrentValue<UMediaPlayer*>(Object), Bindings);
+			return FMediaTexturePropertyPreAnimatedToken(Bindings->GetCurrentValue<UMediaTexture*>(Object), Bindings);
 		}
 
 		TSharedPtr<FTrackInstancePropertyBindings> Bindings;
@@ -154,28 +159,28 @@ public:
  * This pre-animated token is only created if we create a new
  * media texture and assign it to an existing media player.
  */
-struct FMediaPlayerTexturePreAnimatedToken
+struct FMediaTexturePreAnimatedToken
 	: IMovieScenePreAnimatedToken
 {
 	/** Virtual destructor. */
-	virtual ~FMediaPlayerTexturePreAnimatedToken() { }
+	virtual ~FMediaTexturePreAnimatedToken() { }
 
 	virtual void RestoreState(UObject& RestoreObject, IMovieScenePlayer& Player)
 	{
-		UMediaPlayer* MediaPlayer = CastChecked<UMediaPlayer>(&RestoreObject);
+		UMediaTexture* MediaTexture = CastChecked<UMediaTexture>(&RestoreObject);
 		// @todo gmp: fix me media framework 3.0
 		/*
 		MediaPlayer->SetVideoTexture(nullptr);
 		*/
 	}
 
-	static FMovieSceneAnimTypeID GetAnimTypeID() { return TMovieSceneAnimTypeID<FMediaPlayerTexturePreAnimatedToken>(); }
+	static FMovieSceneAnimTypeID GetAnimTypeID() { return TMovieSceneAnimTypeID<FMediaTexturePreAnimatedToken>(); }
 
 	struct FProducer : IMovieScenePreAnimatedTokenProducer
 	{
 		virtual IMovieScenePreAnimatedTokenPtr CacheExistingState(UObject& Object) const override
 		{
-			return FMediaPlayerTexturePreAnimatedToken();
+			return FMediaTexturePreAnimatedToken();
 		}
 	};
 };
@@ -188,74 +193,64 @@ struct FMediaSectionData
 	: PropertyTemplate::FSectionData, FGCObject
 {
 	FMediaSectionData()
-		: TemporaryMediaPlayer(nullptr)
+		: TemporaryMediaTexture(nullptr)
 	{ }
 
 	virtual void AddReferencedObjects(FReferenceCollector& Collector) override
 	{
-		Collector.AddReferencedObject(TemporaryMediaPlayer);
+		Collector.AddReferencedObject(TemporaryMediaTexture);
 	}
 
-	UMediaPlayer* GetTemporaryMediaPlayer()
+	UMediaTexture* GetTemporaryMediaTexture()
 	{
-		if (TemporaryMediaPlayer == nullptr)
+		if (TemporaryMediaTexture == nullptr)
 		{
-			TemporaryMediaPlayer = NewObject<UMediaPlayer>(GetTransientPackage(), MakeUniqueObjectName(GetTransientPackage(), UMediaPlayer::StaticClass()));
+			TemporaryMediaTexture = NewObject<UMediaTexture>(GetTransientPackage(), MakeUniqueObjectName(GetTransientPackage(), UMediaTexture::StaticClass()));
+			TemporaryMediaTexture->UpdateResource();
+
+			UMediaPlayer* MediaPlayer = NewObject<UMediaPlayer>(GetTransientPackage(), MakeUniqueObjectName(GetTransientPackage(), UMediaPlayer::StaticClass()));
+			TemporaryMediaTexture->SetMediaPlayer(MediaPlayer);
 		}
 
-		return TemporaryMediaPlayer;
+		return TemporaryMediaTexture;
 	}
 
-	UMediaPlayer* GetOrUpdateMediaPlayerFromProperty(UObject& InObject, IMovieScenePlayer& Player)
+	UMediaTexture* GetOrUpdateMediaTextureFromProperty(UObject& InObject, IMovieScenePlayer& Player)
 	{
-		UMediaPlayer* MediaPlayer = PropertyBindings->GetCurrentValue<UMediaPlayer*>(InObject);
+		UMediaTexture* MediaTexture = PropertyBindings->GetCurrentValue<UMediaTexture*>(InObject);
 
 		// assign the temporary player to the property if it's null or not the temporary player
-		const bool bNeedToAssignTemporaryPlayer = !MediaPlayer || (TemporaryMediaPlayer && MediaPlayer != TemporaryMediaPlayer);
+		const bool bNeedToAssignTemporaryTexture = !MediaTexture || (TemporaryMediaTexture && MediaTexture != TemporaryMediaTexture);
 
-		if (bNeedToAssignTemporaryPlayer)
+		if (bNeedToAssignTemporaryTexture)
 		{
-			GetTemporaryMediaPlayer();
-			MediaPlayer = TemporaryMediaPlayer;
-		}
+			GetTemporaryMediaTexture();
 
-		check(MediaPlayer);
+			UMediaPlaneComponent* MediaPlaneComponent = CastChecked<UMediaPlaneComponent>(&InObject);
 
-		// Ensure we have a video texture on the media player (transient ones will always
-		// have a texture already, but media players specified on external properties may not)
+			MediaPlaneComponent->SetMediaTexture(MediaTexture);
 
-		// @todo gmp: fix me Media Framework 3.0
-		/*
-		if (!MediaPlayer->GetVideoTexture())
-		{
-			UMediaTexture* NewTexture = NewObject<UMediaTexture>(GetTransientPackage(), MakeUniqueObjectName(GetTransientPackage(), UMediaTexture::StaticClass()));
-			NewTexture->UpdateResource();
-
-			// Save the previous video texture so we don't end up trying to save one that's in a transient package
-			Player.SavePreAnimatedState(*MediaPlayer, FMediaPlayerTexturePreAnimatedToken::GetAnimTypeID(), FMediaPlayerTexturePreAnimatedToken::FProducer());
-
-			MediaPlayer->SetVideoTexture(NewTexture);
+			MediaTexture = TemporaryMediaTexture;
 		}
-		*/
 
 		// always assign the property last, so that anything responding to the update
 		// can have access to the media player and its video texture if needs be
 
-		if (bNeedToAssignTemporaryPlayer)
+		if (bNeedToAssignTemporaryTexture)
 		{
 			// save the previous property value using a unique anim type ID to ensure we don't leave it set after evaluation
-			Player.SavePreAnimatedState(InObject, PropertyID, FMediaPlayerPropertyPreAnimatedToken::FProducer(PropertyBindings));
+			Player.SavePreAnimatedState(InObject, PropertyID, FMediaTexturePropertyPreAnimatedToken::FProducer(PropertyBindings));
 
 			// assign the property
-			PropertyBindings->CallFunction<UMediaPlayer*>(InObject, MediaPlayer);
+			PropertyBindings->CallFunction<UMediaTexture*>(InObject, MediaTexture);
 		}
 
-		return MediaPlayer;
+		return MediaTexture;
 	}
 
 private:
 
-	UMediaPlayer* TemporaryMediaPlayer;
+	UMediaTexture* TemporaryMediaTexture;
 };
 
 
@@ -322,29 +317,34 @@ struct FMediaSectionPreRollExecutionToken
 				continue;
 			}
 
-			UMediaPlayer* MediaPlayer = bForceTemporaryPlayer ? SectionData.GetTemporaryMediaPlayer() : SectionData.GetOrUpdateMediaPlayerFromProperty(*Object, Player);
-			if (!ensure(MediaPlayer))
+			UMediaTexture* MediaTexture = bForceTemporaryPlayer ? SectionData.GetTemporaryMediaTexture() : SectionData.GetOrUpdateMediaTextureFromProperty(*Object, Player);
+			if (!ensure(MediaTexture))
 			{
 				continue;
 			}
 
 			// Save the previous media player state
-			Player.SavePreAnimatedState(*MediaPlayer, TMovieSceneAnimTypeID<FMediaSectionPreRollExecutionToken>(), FMediaPlayerPreAnimatedTokenProducer());
+			//Player.SavePreAnimatedState(*MediaTexture, TMovieSceneAnimTypeID<FMediaSectionPreRollExecutionToken>(), FMediaPlayerPreAnimatedTokenProducer());
 
 			// Open the media source if necessary
-			if (MediaPlayer->GetUrl() != Source->GetUrl())
+			UMediaPlayer* MediaPlayer = MediaTexture->GetMediaPlayer();
+
+			if (MediaPlayer)
 			{
-				if (MediaPlayer->CanPlaySource(Source))
-				{
-					MediaPlayer->OpenSource(Source);
-				}
-				else
+				if (MediaPlayer->GetUrl() != Source->GetUrl())
 				{
-					// @todo: log warning/error
+					if (MediaPlayer->CanPlaySource(Source))
+					{
+						MediaPlayer->OpenSource(Source);
+					}
+					else
+					{
+						// @todo: log warning/error
+					}
 				}
-			}
 
-			OnUpdateFrame(MediaPlayer, Context);
+				OnUpdateFrame(MediaPlayer, Context);
+			}
 		}
 	}
 
@@ -414,7 +414,7 @@ struct FMediaSectionExecutionToken
  *****************************************************************************/
 
 FMovieSceneMediaSectionTemplate::FMovieSceneMediaSectionTemplate(const UMovieSceneMediaSection& InSection, const UMovieSceneMediaTrack& InTrack)
-	: PropertyData(InTrack.GetPropertyName(), InTrack.GetPropertyPath(), NAME_None, "OnMediaPlayerChanged")
+	: PropertyData(InTrack.GetPropertyName(), InTrack.GetPropertyPath(), NAME_None, "OnMediaTextureChanged")
 {
 	Params.Source = InSection.GetMediaSource();
 	Params.SectionStartTime = InSection.GetStartTime();
@@ -464,13 +464,13 @@ void FMovieSceneMediaSectionTemplate::Initialize(const FMovieSceneEvaluationOper
 	}
 
 	FMediaSectionData& SectionData = PersistentData.GetSectionData<FMediaSectionData>();
-	UMediaPlayer* TemporaryMediaPlayer = SectionData.GetTemporaryMediaPlayer();
+	UMediaTexture* TemporaryMediaTexture = SectionData.GetTemporaryMediaTexture();
 
 	for (TWeakObjectPtr<> WeakObject : Player.FindBoundObjects(Operand))
 	{
 		if (UObject* Object = WeakObject.Get())
 		{
-			UMediaPlayer* MediaPlayerFromProperty = SectionData.PropertyBindings->GetCurrentValue<UMediaPlayer*>(*Object);
+			UMediaTexture* MediaTextureFromProperty = SectionData.PropertyBindings->GetCurrentValue<UMediaTexture*>(*Object);
 
 			// Ensure the temporary (prerolling) player has no video texture. It will get
 			// created when the video is actually evaluated for real. This will only be
@@ -481,7 +481,7 @@ void FMovieSceneMediaSectionTemplate::Initialize(const FMovieSceneEvaluationOper
 			TemporaryMediaPlayer->SetVideoTexture(nullptr);
 			*/
 
-			if (MediaPlayerFromProperty == TemporaryMediaPlayer)
+			if (MediaTextureFromProperty == TemporaryMediaTexture)
 			{
 				// Ensure the property is not set if we're prerolling (so we don't display preroll frames)
 				FMovieSceneAnimTypeID PropertyID = SectionData.PropertyID;
diff --git a/Engine/Plugins/Media/MediaCompositing/Source/MediaCompositing/Public/MediaPlaneComponent.h b/Engine/Plugins/Media/MediaCompositing/Source/MediaCompositing/Public/MediaPlaneComponent.h
--- a/Engine/Plugins/Media/MediaCompositing/Source/MediaCompositing/Public/MediaPlaneComponent.h
+++ b/Engine/Plugins/Media/MediaCompositing/Source/MediaCompositing/Public/MediaPlaneComponent.h
@@ -27,7 +27,7 @@ struct MEDIACOMPOSITING_API FMediaPlaneParameters
 	UPROPERTY(EditAnywhere, BlueprintReadWrite, AdvancedDisplay, Category="Media Plane")
 	FName TextureParameterName;
 
-	/** Automaticaly size the plane based on the active camera's lens and filmback settings. Target Camera is found by looking for an active camera component from this component's actor, through its attached parents. */
+	/** Automatically size the plane based on the active camera's lens and filmback settings. Target Camera is found by looking for an active camera component from this component's actor, through its attached parents. */
 	UPROPERTY(EditAnywhere, BlueprintReadWrite, Category="Media Plane")
 	bool bFillScreen;
 
@@ -46,6 +46,10 @@ struct MEDIACOMPOSITING_API FMediaPlaneParameters
 	/** Transient MID to hold the material with the render texture patched in */
 	UPROPERTY(transient)
 	UMaterialInstanceDynamic* DynamicMaterial;
+
+	/** Transient media texture applied to the component */
+	UPROPERTY(transient)
+	UTexture* MediaTexture;
 };
 
 /** 
@@ -84,7 +88,7 @@ public:
 
 	/** Called by sequencer if a texture is changed */
 	UFUNCTION()
-	void OnRenderTextureChanged();
+	void OnMediaTextureChanged();
 
 	/** Access this component's cached view projection matrix. Only valid when the plane is set to fill screen */
 	const FMatrix& GetCachedViewProjectionMatrix() const
@@ -127,6 +131,11 @@ public:
 	 */
 	static void GetProjectionMatricesFromViewTarget(AActor* InViewTarget, FMatrix& OutViewProjectionMatrix, FMatrix& OutInvViewProjectionMatrix);
 
+	/**
+	 * Set the external media texture
+	 */
+	void SetMediaTexture(UTexture* Texture);
+
 protected:
 
 	void UpdateMaterialParametersForMedia();
diff --git a/Engine/Plugins/Media/MediaCompositing/Source/MediaCompositingEditor/Private/Sequencer/MediaThumbnailSection.cpp b/Engine/Plugins/Media/MediaCompositing/Source/MediaCompositingEditor/Private/Sequencer/MediaThumbnailSection.cpp
--- a/Engine/Plugins/Media/MediaCompositing/Source/MediaCompositingEditor/Private/Sequencer/MediaThumbnailSection.cpp
+++ b/Engine/Plugins/Media/MediaCompositing/Source/MediaCompositingEditor/Private/Sequencer/MediaThumbnailSection.cpp
@@ -219,7 +219,7 @@ void FMediaThumbnailSection::Setup()
 	if (MediaTexture == nullptr)
 	{
 		MediaTexture = NewObject<UMediaTexture>(GetTransientPackage(), MakeUniqueObjectName(GetTransientPackage(), UMediaTexture::StaticClass()));
-		MediaTexture->MediaPlayer = MediaPlayer;
+		MediaTexture->SetMediaPlayer(MediaPlayer);
 		MediaTexture->UpdateResource();
 	}
 
diff --git a/Engine/Plugins/Media/MediaCompositing/Source/MediaCompositingEditor/Private/Sequencer/MediaTrackEditor.cpp b/Engine/Plugins/Media/MediaCompositing/Source/MediaCompositingEditor/Private/Sequencer/MediaTrackEditor.cpp
--- a/Engine/Plugins/Media/MediaCompositing/Source/MediaCompositingEditor/Private/Sequencer/MediaTrackEditor.cpp
+++ b/Engine/Plugins/Media/MediaCompositing/Source/MediaCompositingEditor/Private/Sequencer/MediaTrackEditor.cpp
@@ -10,6 +10,7 @@
 #include "ISequencerObjectChangeListener.h"
 #include "MediaPlayer.h"
 #include "MediaSource.h"
+#include "MediaTexture.h"
 #include "MovieSceneBinding.h"
 #include "MovieSceneMediaSection.h"
 #include "MovieSceneMediaTrack.h"
@@ -32,7 +33,7 @@
 
 TArray<FAnimatedPropertyKey, TInlineAllocator<1>> FMediaTrackEditor::GetAnimatedPropertyTypes()
 {
-	return TArray<FAnimatedPropertyKey, TInlineAllocator<1>>({ FAnimatedPropertyKey::FromObjectType(UMediaPlayer::StaticClass()) });
+	return TArray<FAnimatedPropertyKey, TInlineAllocator<1>>({ FAnimatedPropertyKey::FromObjectType(UMediaTexture::StaticClass()) });
 }
 
 
@@ -41,7 +42,7 @@ TArray<FAnimatedPropertyKey, TInlineAllocator<1>> FMediaTrackEditor::GetAnimated
 
 FMediaTrackEditor::FMediaTrackEditor(TSharedRef<ISequencer> InSequencer)
 	: FMovieSceneTrackEditor(InSequencer)
-	, PropertyKey(FAnimatedPropertyKey::FromObjectType(UMediaPlayer::StaticClass()))
+	, PropertyKey(FAnimatedPropertyKey::FromObjectType(UMediaTexture::StaticClass()))
 {
 	ThumbnailPool = MakeShared<FTrackEditorThumbnailPool>(InSequencer);
 	OnPropertyChangedHandle = InSequencer->GetObjectChangeListener().GetOnAnimatablePropertyChanged(PropertyKey).AddRaw(this, &FMediaTrackEditor::OnAnimatedPropertyChanged);
@@ -142,11 +143,14 @@ void FMediaTrackEditor::BuildObjectBindingTrackMenu(FMenuBuilder& MenuBuilder, c
 		check(MediaPlaneComponent);
 
 		UStructProperty* MediaPlanesProperty = MediaPlaneComponent->GetMediaPlaneProperty();
+		UProperty* MediaTextureProperty = FindField<UProperty>(FMediaPlaneParameters::StaticStruct(), GET_MEMBER_NAME_CHECKED(FMediaPlaneParameters, MediaTexture));
 
 		check(MediaPlanesProperty);
+		check(MediaTextureProperty);
 
 		TSharedRef<FPropertyPath> Path = FPropertyPath::CreateEmpty();
 		Path->AddProperty(FPropertyInfo(MediaPlanesProperty));
+		Path->AddProperty(FPropertyInfo(MediaTextureProperty));
 
 		return Path;
 	};
diff --git a/Engine/Plugins/Media/MediaPlayerEditor/Source/MediaPlayerEditor/Private/Factories/MediaPlayerFactoryNew.cpp b/Engine/Plugins/Media/MediaPlayerEditor/Source/MediaPlayerEditor/Private/Factories/MediaPlayerFactoryNew.cpp
--- a/Engine/Plugins/Media/MediaPlayerEditor/Source/MediaPlayerEditor/Private/Factories/MediaPlayerFactoryNew.cpp
+++ b/Engine/Plugins/Media/MediaPlayerEditor/Source/MediaPlayerEditor/Private/Factories/MediaPlayerFactoryNew.cpp
@@ -192,7 +192,7 @@ UObject* UMediaPlayerFactoryNew::FactoryCreateNew(UClass* InClass, UObject* InPa
 
 		if (VideoTexture != nullptr)
 		{
-			VideoTexture->MediaPlayer = NewMediaPlayer;
+			VideoTexture->SetMediaPlayer(NewMediaPlayer);
 		}
 	}
 
diff --git a/Engine/Plugins/Media/MediaPlayerEditor/Source/MediaPlayerEditor/Private/Widgets/SMediaPlayerEditorOutput.cpp b/Engine/Plugins/Media/MediaPlayerEditor/Source/MediaPlayerEditor/Private/Widgets/SMediaPlayerEditorOutput.cpp
--- a/Engine/Plugins/Media/MediaPlayerEditor/Source/MediaPlayerEditor/Private/Widgets/SMediaPlayerEditorOutput.cpp
+++ b/Engine/Plugins/Media/MediaPlayerEditor/Source/MediaPlayerEditor/Private/Widgets/SMediaPlayerEditorOutput.cpp
@@ -69,8 +69,8 @@ void SMediaPlayerEditorOutput::Construct(const FArguments& InArgs, UMediaPlayer&
 
 		if (SoundComponent != nullptr)
 		{
-			SoundComponent->MediaPlayer = &InMediaPlayer;
 			SoundComponent->bIsUISound = true;
+			SoundComponent->SetMediaPlayer(&InMediaPlayer);
 			SoundComponent->Initialize();
 			SoundComponent->AddToRoot();
 		}
@@ -81,7 +81,8 @@ void SMediaPlayerEditorOutput::Construct(const FArguments& InArgs, UMediaPlayer&
 
 	if (MediaTexture != nullptr)
 	{
-		MediaTexture->MediaPlayer = &InMediaPlayer;
+		MediaTexture->AutoClear = true;
+		MediaTexture->SetMediaPlayer(&InMediaPlayer);
 		MediaTexture->UpdateResource();
 		MediaTexture->AddToRoot();
 	}
diff --git a/Engine/Plugins/Media/MediaPlayerEditor/Source/MediaPlayerEditor/Private/Widgets/SMediaPlayerEditorViewport.cpp b/Engine/Plugins/Media/MediaPlayerEditor/Source/MediaPlayerEditor/Private/Widgets/SMediaPlayerEditorViewport.cpp
--- a/Engine/Plugins/Media/MediaPlayerEditor/Source/MediaPlayerEditor/Private/Widgets/SMediaPlayerEditorViewport.cpp
+++ b/Engine/Plugins/Media/MediaPlayerEditor/Source/MediaPlayerEditor/Private/Widgets/SMediaPlayerEditorViewport.cpp
@@ -33,6 +33,8 @@ void SMediaPlayerEditorViewport::Construct(const FArguments& InArgs, UMediaPlaye
 	MediaPlayer = &InMediaPlayer;
 	Style = InStyle;
 
+	const auto OverlayFont = Style->GetFontStyle("MediaPlayerEditor.ViewportFont");
+
 	ChildSlot
 	[
 		SNew(SOverlay)
@@ -69,74 +71,88 @@ void SMediaPlayerEditorViewport::Construct(const FArguments& InArgs, UMediaPlaye
 		+ SOverlay::Slot()
 			.Padding(FMargin(12.0f, 8.0f))
 			[
-				// top info overlays
-				SNew(SHorizontalBox)
+				// info overlays
+				SNew(SVerticalBox)
 					.Visibility_Lambda([this]() -> EVisibility {
 						return (IsHovered() || HasMouseCapture()) ? EVisibility::Visible : EVisibility::Hidden;
 					})
 
-				+ SHorizontalBox::Slot()
-					.HAlign(HAlign_Left)
+				+ SVerticalBox::Slot()
 					.VAlign(VAlign_Top)
 					[
-						// media source name
-						SNew(STextBlock)
-							.ColorAndOpacity(FSlateColor::UseSubduedForeground())
-							.Font(Style->GetFontStyle("MediaPlayerEditor.ViewportFont"))
-							.ShadowOffset(FVector2D(1.f, 1.f))
-							.Text(this, &SMediaPlayerEditorViewport::HandleMediaSourceNameText)
-							.ToolTipText(LOCTEXT("OverlaySourceNameTooltip", "Name of the currently opened media source"))
+						SNew(SHorizontalBox)
+
+						+ SHorizontalBox::Slot()
+							.HAlign(HAlign_Left)
+							.VAlign(VAlign_Top)
+							[
+								// media source name
+								SNew(STextBlock)
+									.ColorAndOpacity(FSlateColor::UseSubduedForeground())
+									.Font(OverlayFont)
+									.ShadowOffset(FVector2D(1.f, 1.f))
+									.Text(this, &SMediaPlayerEditorViewport::HandleMediaSourceNameText)
+									.ToolTipText(LOCTEXT("OverlaySourceNameTooltip", "Name of the currently opened media source"))
+							]
+
+						+ SHorizontalBox::Slot()
+							.HAlign(HAlign_Right)
+							.VAlign(VAlign_Top)
+							[
+								// player name
+								SNew(STextBlock)
+									.ColorAndOpacity(FSlateColor::UseSubduedForeground())
+									.Font(OverlayFont)
+									.ShadowOffset(FVector2D(1.f, 1.f))
+									.Text(this, &SMediaPlayerEditorViewport::HandleMediaPlayerNameText)
+									.ToolTipText(LOCTEXT("OverlayPlayerNameTooltip", "Name of the currently used media player plug-in"))
+							]
 					]
 
-				+ SHorizontalBox::Slot()
-					.HAlign(HAlign_Right)
-					.VAlign(VAlign_Top)
+				+ SVerticalBox::Slot()
+					.VAlign(VAlign_Bottom)
 					[
-						// player name
-						SNew(STextBlock)
-							.ColorAndOpacity(FSlateColor::UseSubduedForeground())
-							.Font(Style->GetFontStyle("MediaPlayerEditor.ViewportFont"))
-							.ShadowOffset(FVector2D(1.f, 1.f))
-							.Text(this, &SMediaPlayerEditorViewport::HandleMediaPlayerNameText)
-							.ToolTipText(LOCTEXT("OverlayPlayerNameTooltip", "Name of the currently used media player plug-in"))
+						SNew(SHorizontalBox)
+
+						+ SHorizontalBox::Slot()
+							.HAlign(HAlign_Left)
+							.VAlign(VAlign_Bottom)
+							[
+								// playback state
+								SNew(STextBlock)
+									.ColorAndOpacity(FSlateColor::UseSubduedForeground())
+									.Font(OverlayFont)
+									.ShadowOffset(FVector2D(1.f, 1.f))
+									.Text(this, &SMediaPlayerEditorViewport::HandleMediaPlayerStateText)
+									.ToolTipText(LOCTEXT("OverlayPlayerStateTooltip", "The media player's current state"))
+							]
+
+						+ SHorizontalBox::Slot()
+							.HAlign(HAlign_Right)
+							.VAlign(VAlign_Bottom)
+							[
+								// view settings
+								SNew(STextBlock)
+									.ColorAndOpacity(FSlateColor::UseSubduedForeground())
+									.Font(OverlayFont)
+									.ShadowOffset(FVector2D(1.f, 1.f))
+									.Text(this, &SMediaPlayerEditorViewport::HandleViewSettingsText)
+									.ToolTipText(LOCTEXT("OverlayViewSettingsTooltip", "The current view settings"))
+							]
 					]
 			]
 
 		+ SOverlay::Slot()
 			.Padding(FMargin(12.0f, 8.0f))
-			.VAlign(VAlign_Bottom)
+			.HAlign(HAlign_Center)
+			.VAlign(VAlign_Center)
 			[
-				// bottom info overlays
-				SNew(SHorizontalBox)
-					.Visibility_Lambda([this]() -> EVisibility {
-						return (IsHovered() || HasMouseCapture()) ? EVisibility::Visible : EVisibility::Hidden;
-					})
-
-				+ SHorizontalBox::Slot()
-					.HAlign(HAlign_Left)
-					.VAlign(VAlign_Bottom)
-					[
-						// playback state
-						SNew(STextBlock)
-							.ColorAndOpacity(FSlateColor::UseSubduedForeground())
-							.Font(Style->GetFontStyle("MediaPlayerEditor.ViewportFont"))
-							.ShadowOffset(FVector2D(1.f, 1.f))
-							.Text(this, &SMediaPlayerEditorViewport::HandleMediaPlayerStateText)
-							.ToolTipText(LOCTEXT("OverlayPlayerStateTooltip", "The media player's current state"))
-					]
-
-				+ SHorizontalBox::Slot()
-					.HAlign(HAlign_Right)
-					.VAlign(VAlign_Bottom)
-					[
-						// view settings
-						SNew(STextBlock)
-							.ColorAndOpacity(FSlateColor::UseSubduedForeground())
-							.Font(Style->GetFontStyle("MediaPlayerEditor.ViewportFont"))
-							.ShadowOffset(FVector2D(1.f, 1.f))
-							.Text(this, &SMediaPlayerEditorViewport::HandleViewSettingsText)
-							.ToolTipText(LOCTEXT("OverlayViewSettingsTooltip", "The current view settings"))
-					]
+				// notification overlay
+				SNew(STextBlock)
+					.ColorAndOpacity(FSlateColor::UseSubduedForeground())
+					.Font(OverlayFont)
+					.ShadowOffset(FVector2D(1.f, 1.f))
+					.Text(this, &SMediaPlayerEditorViewport::HandleNotificationText)
 			]
 	];
 }
@@ -297,6 +313,37 @@ FText SMediaPlayerEditorViewport::HandleMediaSourceNameText() const
 }
 
 
+FText SMediaPlayerEditorViewport::HandleNotificationText() const
+{
+	if (MediaPlayer->IsPlaying())
+	{
+		if (MediaPlayer->GetNumTracks(EMediaPlayerTrack::Video) == 0)
+		{
+			return LOCTEXT("NoVideoTrackAvailable", "No video track available");
+		}
+
+		const int32 SelectedVideoTrack = MediaPlayer->GetSelectedTrack(EMediaPlayerTrack::Video);
+
+		if (SelectedVideoTrack == INDEX_NONE)
+		{
+			return LOCTEXT("NoVideoTrackSelected", "No video track selected");
+		}
+
+		if (MediaPlayer->GetNumTrackFormats(EMediaPlayerTrack::Video, SelectedVideoTrack) == 0)
+		{
+			return LOCTEXT("NoVideoFormatsAvailable", "No video formats available");
+		}
+
+		if (MediaPlayer->GetTrackFormat(EMediaPlayerTrack::Video, SelectedVideoTrack) == INDEX_NONE)
+		{
+			return LOCTEXT("NoVideoFormatSelected", "No video format selected");
+		}
+	}
+
+	return FText::GetEmpty();
+}
+
+
 FText SMediaPlayerEditorViewport::HandleViewSettingsText() const
 {
 	FVector Euler = MediaPlayer->GetViewRotation().Euler();
diff --git a/Engine/Plugins/Media/MediaPlayerEditor/Source/MediaPlayerEditor/Private/Widgets/SMediaPlayerEditorViewport.h b/Engine/Plugins/Media/MediaPlayerEditor/Source/MediaPlayerEditor/Private/Widgets/SMediaPlayerEditorViewport.h
--- a/Engine/Plugins/Media/MediaPlayerEditor/Source/MediaPlayerEditor/Private/Widgets/SMediaPlayerEditorViewport.h
+++ b/Engine/Plugins/Media/MediaPlayerEditor/Source/MediaPlayerEditor/Private/Widgets/SMediaPlayerEditorViewport.h
@@ -54,6 +54,9 @@ private:
 	/** Callback for getting the text of the media source name overlay. */
 	FText HandleMediaSourceNameText() const;
 
+	/** Callback for getting the text of the notification overlay. */
+	FText HandleNotificationText() const;
+
 	/** Callback for getting the text of the view settings overlay. */
 	FText HandleViewSettingsText() const;
 
diff --git a/Engine/Plugins/Media/MfMedia/Source/MfMedia/Private/Player/MfMediaPlayer.cpp b/Engine/Plugins/Media/MfMedia/Source/MfMedia/Private/Player/MfMediaPlayer.cpp
--- a/Engine/Plugins/Media/MfMedia/Source/MfMedia/Private/Player/MfMediaPlayer.cpp
+++ b/Engine/Plugins/Media/MfMedia/Source/MfMedia/Private/Player/MfMediaPlayer.cpp
@@ -20,8 +20,6 @@
 	#include "XboxOneAllowPlatformTypes.h"
 #endif
 
-#define MFMEDIAPLAYER_USE_SEEKANDREVERSE 0 // not fully working yet
-
 #define LOCTEXT_NAMESPACE "FMfMediaPlayer"
 
 
@@ -254,9 +252,14 @@ void FMfMediaPlayer::TickFetch(FTimespan /*DeltaTime*/, FTimespan /*Timecode*/)
 
 void FMfMediaPlayer::TickInput(FTimespan DeltaTime, FTimespan Timecode)
 {
-	if ((CurrentState != EMediaState::Playing) || (CurrentDuration == FTimespan::Zero()))
+	if (CurrentDuration == FTimespan::Zero())
 	{
-		return;
+		return; // nothing to play
+	}
+
+	if ((CurrentState != EMediaState::Paused) && (CurrentState != EMediaState::Playing))
+	{
+		return; // not paused or playing
 	}
 
 	if (SourceReaderError)
@@ -308,19 +311,9 @@ void FMfMediaPlayer::TickInput(FTimespan DeltaTime, FTimespan Timecode)
 			Tracks->Restart();
 		}
 	}
-	else if (CurrentRate < 0.0f)
-	{
-		// IMFSourceReader does not support reverse playback, even if the media
-		// source does. we emulate the WMF behavior by seeking to key frames.
-
-		if (!CommitTime(CurrentTime))
-		{
-			CurrentState = EMediaState::Error;
-		}
-	}
 
 	// update tracks
-	if (CurrentState == EMediaState::Playing)
+	if ((CurrentState == EMediaState::Paused) || (CurrentState == EMediaState::Playing))
 	{
 		Tracks->TickInput(CurrentRate, CurrentTime);
 	}
@@ -354,6 +347,7 @@ bool FMfMediaPlayer::CommitTime(FTimespan Time)
 	}
 
 	CurrentTime = Time;
+	Tracks->Restart();
 
 	return true;
 }
@@ -452,6 +446,9 @@ void FMfMediaPlayer::UpdateCharacteristics()
 	// cache rate control properties
 	if (RateSupport.IsValid())
 	{
+		// Note: IMFSourceReader does not support reverse playback, even if the
+		// media source does, so we are only querying the supported forward rates.
+
 		float MaxRate = 0.0f;
 		float MinRate = 0.0f;
 
@@ -461,27 +458,11 @@ void FMfMediaPlayer::UpdateCharacteristics()
 			ThinnedRates.Add(TRange<float>::Inclusive(MinRate, MaxRate));
 		}
 
-#if MFMEDIAPLAYER_USE_SEEKANDREVERSE
-		if (SUCCEEDED(RateSupport->GetSlowestRate(MFRATE_REVERSE, TRUE, &MinRate)) &&
-			SUCCEEDED(RateSupport->GetFastestRate(MFRATE_REVERSE, TRUE, &MaxRate)))
-		{
-			ThinnedRates.Add(TRange<float>::Inclusive(MaxRate, MinRate));
-		}
-#endif
-
 		if (SUCCEEDED(RateSupport->GetSlowestRate(MFRATE_FORWARD, FALSE, &MinRate)) &&
 			SUCCEEDED(RateSupport->GetFastestRate(MFRATE_FORWARD, FALSE, &MaxRate)))
 		{
 			UnthinnedRates.Add(TRange<float>::Inclusive(MinRate, MaxRate));
 		}
-
-#if MFMEDIAPLAYER_USE_SEEKANDREVERSE
-		if (SUCCEEDED(RateSupport->GetSlowestRate(MFRATE_REVERSE, FALSE, &MinRate)) &&
-			SUCCEEDED(RateSupport->GetFastestRate(MFRATE_REVERSE, FALSE, &MaxRate)))
-		{
-			UnthinnedRates.Add(TRange<float>::Inclusive(MaxRate, MinRate));
-		}
-#endif
 	}
 }
 
@@ -506,7 +487,6 @@ bool FMfMediaPlayer::CanControl(EMediaControl Control) const
 		return ((CurrentState != EMediaState::Playing) && ThinnedRates.Contains(1.0f));
 	}
 
-#if MFMEDIAPLAYER_USE_SEEKANDREVERSE
 	if (Control == EMediaControl::Scrub)
 	{
 		return (((Characteristics & MFMEDIASOURCE_HAS_SLOW_SEEK) == 0) && (ThinnedRates.Contains(0.0f)));
@@ -516,7 +496,6 @@ bool FMfMediaPlayer::CanControl(EMediaControl Control) const
 	{
 		return (((Characteristics & MFMEDIASOURCE_CAN_SEEK) != 0) && (CurrentDuration > FTimespan::Zero()));
 	}
-#endif
 
 	return false;
 }
diff --git a/Engine/Plugins/Media/MfMedia/Source/MfMedia/Private/Player/MfMediaTracks.cpp b/Engine/Plugins/Media/MfMedia/Source/MfMedia/Private/Player/MfMediaTracks.cpp
--- a/Engine/Plugins/Media/MfMedia/Source/MfMedia/Private/Player/MfMediaTracks.cpp
+++ b/Engine/Plugins/Media/MfMedia/Source/MfMedia/Private/Player/MfMediaTracks.cpp
@@ -119,15 +119,25 @@ FTimespan FMfMediaTracks::GetDuration() const
 {
 	FScopeLock Lock(&CriticalSection);
 
-	if (PresentationDescriptor == NULL)
+	if (!SourceReader.IsValid())
 	{
 		return FTimespan::Zero();
 	}
 
-	UINT64 PresentationDuration = 0;
-	PresentationDescriptor->GetUINT64(MF_PD_DURATION, &PresentationDuration);
+	PROPVARIANT DurationAttrib;
+	{
+		const HRESULT Result = SourceReader->GetPresentationAttribute(MF_SOURCE_READER_MEDIASOURCE, MF_PD_DURATION, &DurationAttrib);
+
+		if (FAILED(Result))
+		{
+			return FTimespan::Zero();
+		}
+	}
 
-	return FTimespan(PresentationDuration);
+	const int64 Duration = (DurationAttrib.vt == VT_UI8) ? (int64)DurationAttrib.uhVal.QuadPart : 0;
+	::PropVariantClear(&DurationAttrib);
+
+	return FTimespan(Duration);
 }
 
 
@@ -169,18 +179,6 @@ void FMfMediaTracks::Initialize(IMFMediaSource* InMediaSource, IMFSourceReaderCa
 		return;
 	}
 
-	// create presentation descriptor
-	TComPtr<IMFPresentationDescriptor> NewPresentationDescriptor;
-	{
-		HRESULT Result = InMediaSource->CreatePresentationDescriptor(&NewPresentationDescriptor);
-
-		if (FAILED(Result))
-		{
-			UE_LOG(LogMfMedia, Verbose, TEXT("Tracks %p: Failed to create presentation descriptor: %s"), this, *MfMedia::ResultToString(Result));
-			return;
-		}
-	}
-
 	// create source reader
 	TComPtr<IMFAttributes> Attributes;
 	{
@@ -213,7 +211,7 @@ void FMfMediaTracks::Initialize(IMFMediaSource* InMediaSource, IMFSourceReaderCa
 
 	TComPtr<IMFSourceReader> NewSourceReader;
 	{
-		HRESULT Result = ::MFCreateSourceReaderFromMediaSource(InMediaSource, Attributes, &NewSourceReader);
+		const HRESULT Result = ::MFCreateSourceReaderFromMediaSource(InMediaSource, Attributes, &NewSourceReader);
 
 		if (FAILED(Result))
 		{
@@ -222,39 +220,25 @@ void FMfMediaTracks::Initialize(IMFMediaSource* InMediaSource, IMFSourceReaderCa
 		}
 	}
 
-	// get number of streams
-	DWORD StreamCount = 0;
-	{
-		HRESULT Result = NewPresentationDescriptor->GetStreamDescriptorCount(&StreamCount);
-
-		if (FAILED(Result))
-		{
-			UE_LOG(LogMfMedia, Verbose, TEXT("Tracks %p: Failed to get stream count: %s"), this, *MfMedia::ResultToString(Result));
-			return;
-		}
-
-		UE_LOG(LogMfMedia, Verbose, TEXT("Tracks %p: Found %i streams"), this, StreamCount);
-	}
-
 	// initialization successful
 	MediaSource = InMediaSource;
-	PresentationDescriptor = NewPresentationDescriptor;
 	Samples = InSamples;
 	SourceReader = NewSourceReader;
 
-	// add streams (Media Foundation reports them in reverse order)
-	bool AllStreamsAdded = true;
+	// add streams
+	int32 StreamIndex = 0;
 
-	for (int32 StreamIndex = StreamCount - 1; StreamIndex >= 0; --StreamIndex)
+	while (AddStreamToTracks(StreamIndex, Info))
 	{
-		AllStreamsAdded &= AddStreamToTracks(StreamIndex, Info);
 		Info += TEXT("\n");
+		++StreamIndex;
 	}
 
-	if (!AllStreamsAdded)
-	{
-		UE_LOG(LogMfMedia, Verbose, TEXT("Tracks %p: Not all available streams were added to the track collection"), this);
-	}
+	UE_LOG(LogMfMedia, Verbose, TEXT("Tracks %p: Found %i streams"), this, StreamIndex);
+
+	Algo::Reverse(AudioTracks);
+	Algo::Reverse(CaptionTracks);
+	Algo::Reverse(VideoTracks);
 }
 
 
@@ -449,7 +433,6 @@ void FMfMediaTracks::Shutdown()
 	CaptionDone = true;
 	VideoDone = true;
 
-	PresentationDescriptor.Reset();
 	Samples.Reset();
 	SourceReader.Reset();
 
@@ -781,20 +764,14 @@ bool FMfMediaTracks::SelectTrack(EMediaTrackType TrackType, int32 TrackIndex)
 	if (*SelectedTrack != INDEX_NONE)
 	{
 		const DWORD StreamIndex = (*Tracks)[*SelectedTrack].StreamIndex;
-		HRESULT Result = PresentationDescriptor->DeselectStream(StreamIndex);
-
-		if (FAILED(Result))
 		{
-			UE_LOG(LogMfMedia, Verbose, TEXT("Tracks %p: Failed to deselect stream %i on presentation descriptor: %s"), this, StreamIndex, *MfMedia::ResultToString(Result));
-			return false;
-		}
+			const HRESULT Result = SourceReader->SetStreamSelection(StreamIndex, FALSE);
 
-		Result = SourceReader->SetStreamSelection(StreamIndex, FALSE);
-
-		if (FAILED(Result))
-		{
-			UE_LOG(LogMfMedia, Verbose, TEXT("Tracks %p: Failed to deselect stream %i on source reader: %s"), this, StreamIndex, *MfMedia::ResultToString(Result));
-			return false;
+			if (FAILED(Result))
+			{
+				UE_LOG(LogMfMedia, Verbose, TEXT("Tracks %p: Failed to deselect stream %i on source reader: %s"), this, StreamIndex, *MfMedia::ResultToString(Result));
+				return false;
+			}
 		}
 
 		UE_LOG(LogMfMedia, Verbose, TEXT("Tracks %p: Disabled stream %i"), this, StreamIndex);
@@ -802,7 +779,7 @@ bool FMfMediaTracks::SelectTrack(EMediaTrackType TrackType, int32 TrackIndex)
 		*SelectedTrack = INDEX_NONE;
 		SelectionChanged = true;
 
-		Result = SourceReader->Flush(StreamIndex);
+		HRESULT Result = SourceReader->Flush(StreamIndex);
 
 		if (FAILED(Result))
 		{
@@ -814,7 +791,7 @@ bool FMfMediaTracks::SelectTrack(EMediaTrackType TrackType, int32 TrackIndex)
 	if (TrackIndex != INDEX_NONE)
 	{
 		const DWORD StreamIndex = (*Tracks)[TrackIndex].StreamIndex;
-		HRESULT Result = SourceReader->SetStreamSelection(StreamIndex, TRUE);
+		const HRESULT Result = SourceReader->SetStreamSelection(StreamIndex, TRUE);
 
 		if (FAILED(Result))
 		{
@@ -822,14 +799,6 @@ bool FMfMediaTracks::SelectTrack(EMediaTrackType TrackType, int32 TrackIndex)
 			return false;
 		}
 
-		Result = PresentationDescriptor->SelectStream(StreamIndex);
-
-		if (FAILED(Result))
-		{
-			UE_LOG(LogMfMedia, Verbose, TEXT("Tracks %p: Failed to select stream %i on presenation descriptor: %s"), this, StreamIndex, *MfMedia::ResultToString(Result));
-			return false;
-		}
-
 		UE_LOG(LogMfMedia, Verbose, TEXT("Tracks %p: Enabled stream %i"), this, StreamIndex);
 			
 		*SelectedTrack = TrackIndex;
@@ -889,25 +858,17 @@ bool FMfMediaTracks::SetTrackFormat(EMediaTrackType TrackType, int32 TrackIndex,
 
 	check(Format.InputType.IsValid());
 	check(Format.OutputType.IsValid());
-	check(Track.Handler.IsValid());
 
 	// set track format
-	UE_LOG(LogMfMedia, Verbose, TEXT("Tracks %p: Set format %i instead of %i on %s track %i (%i formats)"), this, FormatIndex, Track.SelectedFormat, *MediaUtils::TrackTypeToString(TrackType), TrackIndex, Track.Formats.Num());
-
-	HRESULT Result = Track.Handler->SetCurrentMediaType(Format.InputType);
-
-	if (FAILED(Result))
+	UE_LOG(LogMfMedia, Verbose, TEXT("Tracks %p: Setting format %i instead of %i on %s track %i (%i formats)"), this, FormatIndex, Track.SelectedFormat, *MediaUtils::TrackTypeToString(TrackType), TrackIndex, Track.Formats.Num());
 	{
-		UE_LOG(LogMfMedia, Verbose, TEXT("Tracks %p: Failed to set selected media type on handler for stream %i: %s"), this, Track.StreamIndex, *MfMedia::ResultToString(Result));
-		return false;
-	}
-
-	Result = SourceReader->SetCurrentMediaType(Track.StreamIndex, NULL, Format.OutputType);
+		HRESULT Result = SourceReader->SetCurrentMediaType(Track.StreamIndex, NULL, Format.OutputType);
 
-	if (FAILED(Result))
-	{
-		UE_LOG(LogMfMedia, Verbose, TEXT("Tracks %p: Failed to set selected media type on reader for stream %i: %s"), this, Track.StreamIndex, *MfMedia::ResultToString(Result));
-		return false;
+		if (FAILED(Result))
+		{
+			UE_LOG(LogMfMedia, Verbose, TEXT("Tracks %p: Failed to set selected media type on reader for stream %i: %s"), this, Track.StreamIndex, *MfMedia::ResultToString(Result));
+			return false;
+		}
 	}
 
 	Track.SelectedFormat = FormatIndex;
@@ -924,35 +885,15 @@ bool FMfMediaTracks::AddStreamToTracks(uint32 StreamIndex, FString& OutInfo)
 {
 	OutInfo += FString::Printf(TEXT("Stream %i\n"), StreamIndex);
 
-	// get stream descriptor
-	TComPtr<IMFStreamDescriptor> StreamDescriptor;
-	{
-		BOOL Selected = FALSE;
-		const HRESULT Result = PresentationDescriptor->GetStreamDescriptorByIndex(StreamIndex, &Selected, &StreamDescriptor);
-
-		if (FAILED(Result))
-		{
-			UE_LOG(LogMfMedia, Verbose, TEXT("Tracks %p: Failed to get stream descriptor for stream %i: %s"), this, StreamIndex, *MfMedia::ResultToString(Result));
-			OutInfo += TEXT("\tmissing stream descriptor\n");
-
-			return false;
-		}
-
-		if (Selected == TRUE)
-		{
-			PresentationDescriptor->DeselectStream(StreamIndex);
-		}
-	}
-
-	// get media type handler
-	TComPtr<IMFMediaTypeHandler> Handler;
+	// get current format
+	TComPtr<IMFMediaType> CurrentMediaType;
 	{
-		const HRESULT Result = StreamDescriptor->GetMediaTypeHandler(&Handler);
+		const HRESULT Result = SourceReader->GetCurrentMediaType(StreamIndex, &CurrentMediaType);
 
 		if (FAILED(Result))
 		{
-			UE_LOG(LogMfMedia, Verbose, TEXT("Tracks %p: Failed to get media type handler for stream %i: %s"), this, StreamIndex, *MfMedia::ResultToString(Result));
-			OutInfo += TEXT("\tno handler available\n");
+			UE_LOG(LogMfMedia, Verbose, TEXT("Tracks %p: Failed to get current media type in stream %i: %s"), this, StreamIndex, *MfMedia::ResultToString(Result));
+			OutInfo += TEXT("\tfailed to get current format\n");
 
 			return false;
 		}
@@ -961,7 +902,7 @@ bool FMfMediaTracks::AddStreamToTracks(uint32 StreamIndex, FString& OutInfo)
 	// skip unsupported handler types
 	GUID MajorType;
 	{
-		const HRESULT Result = Handler->GetMajorType(&MajorType);
+		const HRESULT Result = CurrentMediaType->GetGUID(MF_MT_MAJOR_TYPE, &MajorType);
 
 		if (FAILED(Result))
 		{
@@ -986,25 +927,19 @@ bool FMfMediaTracks::AddStreamToTracks(uint32 StreamIndex, FString& OutInfo)
 	}
 
 	// @todo gmp: handle protected content
-	const bool Protected = ::MFGetAttributeUINT32(StreamDescriptor, MF_SD_PROTECTED, FALSE) != 0;
+	BOOL Protected = FALSE;
 	{
-		if (Protected)
-		{
-			OutInfo += FString::Printf(TEXT("\tProtected content\n"));
-		}
-	}
-
-	// get number of track formats
-	DWORD NumMediaTypes = 0;
-	{
-		const HRESULT Result = Handler->GetMediaTypeCount(&NumMediaTypes);
+		PROPVARIANT ProtectedAttrib;
 
-		if (FAILED(Result))
+		if (SUCCEEDED(SourceReader->GetPresentationAttribute(StreamIndex, MF_SD_PROTECTED, &ProtectedAttrib)))
 		{
-			UE_LOG(LogMfMedia, Verbose, TEXT("Tracks %p: Failed to get number of track formats in stream %i"), this, StreamIndex);
-			OutInfo += TEXT("\tfailed to get track formats\n");
+			Protected = (ProtectedAttrib.vt == VT_BOOL) && (ProtectedAttrib.boolVal == VARIANT_TRUE);
+			::PropVariantClear(&ProtectedAttrib);
 
-			return false;
+			if (Protected == TRUE)
+			{
+				OutInfo += FString::Printf(TEXT("\tProtected content\n"));
+			}
 		}
 	}
 
@@ -1028,35 +963,35 @@ bool FMfMediaTracks::AddStreamToTracks(uint32 StreamIndex, FString& OutInfo)
 	}
 
 	check(Track != nullptr);
-
-	// get current format
-	TComPtr<IMFMediaType> CurrentMediaType;
-	{
-		HRESULT Result = Handler->GetCurrentMediaType(&CurrentMediaType);
-
-		if (FAILED(Result))
-		{
-			UE_LOG(LogMfMedia, Verbose, TEXT("Tracks %p: Failed to get current media type in stream %i: %s"), this, StreamIndex, *MfMedia::ResultToString(Result));
-		}
-	}
-
 	Track->SelectedFormat = INDEX_NONE;
 
 	// add track formats
 	const bool AllowNonStandardCodecs = false;// GetDefault<UMfMediaSettings>()->AllowNonStandardCodecs;
 
-	for (DWORD TypeIndex = 0; TypeIndex < NumMediaTypes; ++TypeIndex)
+	int32 TypeIndex = INDEX_NONE;
+
+	while (true)
 	{
-		OutInfo += FString::Printf(TEXT("\tFormat %i\n"), TypeIndex);
+		++TypeIndex;
 
 		// get media type
 		TComPtr<IMFMediaType> MediaType;
-
-		if (FAILED(Handler->GetMediaTypeByIndex(TypeIndex, &MediaType)))
 		{
-			OutInfo += TEXT("\t\tfailed to get media type\n");
+			const HRESULT Result = SourceReader->GetNativeMediaType(StreamIndex, TypeIndex, &MediaType);
 
-			continue;
+			if (Result == MF_E_NO_MORE_TYPES)
+			{
+				break; // done
+			}
+
+			OutInfo += FString::Printf(TEXT("\tFormat %i\n"), TypeIndex);
+
+			if (FAILED(Result))
+			{
+				OutInfo += TEXT("\t\tfailed to get media type\n");
+
+				continue;
+			}
 		}
 
 		// get sub-type
@@ -1359,10 +1294,9 @@ bool FMfMediaTracks::AddStreamToTracks(uint32 StreamIndex, FString& OutInfo)
 		for (int32 FormatIndex = 0; FormatIndex < Track->Formats.Num(); ++FormatIndex)
 		{
 			const FFormat& Format = Track->Formats[FormatIndex];
-			const HRESULT ReaderResult = SourceReader->SetCurrentMediaType(StreamIndex, NULL, Format.OutputType);
-			const HRESULT TrackResult = Handler->SetCurrentMediaType(Format.InputType);
+			const HRESULT Result = SourceReader->SetCurrentMediaType(StreamIndex, NULL, Format.OutputType);
 
-			if (SUCCEEDED(TrackResult) && SUCCEEDED(ReaderResult))
+			if (SUCCEEDED(Result))
 			{
 				UE_LOG(LogMfMedia, Verbose, TEXT("Tracks %p: Picked default format %i for stream %i"), this, FormatIndex, StreamIndex);
 				Track->SelectedFormat = FormatIndex;
@@ -1378,28 +1312,41 @@ bool FMfMediaTracks::AddStreamToTracks(uint32 StreamIndex, FString& OutInfo)
 	}
 
 	// set track details
-	PWSTR OutString = NULL;
-	UINT32 OutLength = 0;
+	PROPVARIANT LanguageAttrib;
 
-	if (SUCCEEDED(StreamDescriptor->GetAllocatedString(MF_SD_LANGUAGE, &OutString, &OutLength)))
+	if (SUCCEEDED(SourceReader->GetPresentationAttribute(StreamIndex, MF_SD_LANGUAGE, &LanguageAttrib)))
 	{
-		Track->Language = OutString;
-		::CoTaskMemFree(OutString);
+		if ((LanguageAttrib.vt == VT_LPWSTR) || (LanguageAttrib.vt == VT_BSTR))
+		{
+			if (LanguageAttrib.pwszVal != NULL)
+			{
+				Track->Language = LanguageAttrib.pwszVal;
+			}
+		}
+
+		::PropVariantClear(&LanguageAttrib);
 	}
 
-	if (SUCCEEDED(StreamDescriptor->GetAllocatedString(MF_SD_STREAM_NAME, &OutString, &OutLength)))
+	PROPVARIANT NameAttrib;
+
+	if (SUCCEEDED(SourceReader->GetPresentationAttribute(StreamIndex, MF_SD_STREAM_NAME, &NameAttrib)))
 	{
-		Track->Name = OutString;
-		::CoTaskMemFree(OutString);
+		if ((NameAttrib.vt == VT_LPWSTR) || (NameAttrib.vt == VT_BSTR))
+		{
+			if (NameAttrib.pwszVal != NULL)
+			{
+				Track->Name = NameAttrib.pwszVal;
+			}
+		}
+
+		::PropVariantClear(&NameAttrib);
 	}
 
 	Track->DisplayName = (Track->Name.IsEmpty())
 		? FText::Format(LOCTEXT("UnnamedStreamFormat", "Unnamed Track (Stream {0})"), FText::AsNumber((uint32)StreamIndex))
 		: FText::FromString(Track->Name);
 
-	Track->Descriptor = StreamDescriptor;
-	Track->Handler = Handler;
-	Track->Protected = Protected;
+	Track->Protected = (Protected == TRUE);
 	Track->StreamIndex = StreamIndex;
 
 	return true;
diff --git a/Engine/Plugins/Media/MfMedia/Source/MfMedia/Private/Player/MfMediaTracks.h b/Engine/Plugins/Media/MfMedia/Source/MfMedia/Private/Player/MfMediaTracks.h
--- a/Engine/Plugins/Media/MfMedia/Source/MfMedia/Private/Player/MfMediaTracks.h
+++ b/Engine/Plugins/Media/MfMedia/Source/MfMedia/Private/Player/MfMediaTracks.h
@@ -68,10 +68,8 @@ class FMfMediaTracks
 
 	struct FTrack
 	{
-		TComPtr<IMFStreamDescriptor> Descriptor;
 		FText DisplayName;
 		TArray<FFormat> Formats;
-		TComPtr<IMFMediaTypeHandler> Handler;
 		FString Language;
 		FString Name;
 		bool Protected;
@@ -338,9 +336,6 @@ private:
 	/** Whether the media source has changed. */
 	bool MediaSourceChanged;
 
-	/** The presentation descriptor of the currently opened media. */
-	TComPtr<IMFPresentationDescriptor> PresentationDescriptor;
-
 	/** Media sample collection that receives the output. */
 	TSharedPtr<FMediaSamples, ESPMode::ThreadSafe> Samples;
 
diff --git a/Engine/Plugins/Media/WmfMedia/Source/WmfMedia/Private/Player/WmfMediaPlayer.cpp b/Engine/Plugins/Media/WmfMedia/Source/WmfMedia/Private/Player/WmfMediaPlayer.cpp
--- a/Engine/Plugins/Media/WmfMedia/Source/WmfMedia/Private/Player/WmfMediaPlayer.cpp
+++ b/Engine/Plugins/Media/WmfMedia/Source/WmfMedia/Private/Player/WmfMediaPlayer.cpp
@@ -168,8 +168,6 @@ void FWmfMediaPlayer::TickFetch(FTimespan /*DeltaTime*/, FTimespan /*Timecode*/)
 
 	if (TrackSelectionChanged)
 	{
-		UE_LOG(LogWmfMedia, Verbose, TEXT("Player %p: Creating and setting new playback topology"), this);
-
 		// less than windows 10, seem to be a problem switching stream
 		if (!FWindowsPlatformMisc::VerifyWindowsVersion(10, 0) /* Anything < Windows 10.0 */)
 		{
diff --git a/Engine/Plugins/Media/WmfMedia/Source/WmfMedia/Private/Player/WmfMediaTracks.cpp b/Engine/Plugins/Media/WmfMedia/Source/WmfMedia/Private/Player/WmfMediaTracks.cpp
--- a/Engine/Plugins/Media/WmfMedia/Source/WmfMedia/Private/Player/WmfMediaTracks.cpp
+++ b/Engine/Plugins/Media/WmfMedia/Source/WmfMedia/Private/Player/WmfMediaTracks.cpp
@@ -163,6 +163,8 @@ TComPtr<IMFTopology> FWmfMediaTracks::CreateTopology()
 		return NULL;
 	}
 
+	UE_LOG(LogWmfMedia, Verbose, TEXT("Tracks %p: Created playback topology %p (media source %p)"), this, Topology.Get(), MediaSource.Get());
+
 	return Topology;
 }
 
@@ -196,7 +198,7 @@ void FWmfMediaTracks::Initialize(IMFMediaSource* InMediaSource, const FString& U
 {
 	Shutdown();
 
-	UE_LOG(LogWmfMedia, Verbose, TEXT("Tracks: %p: Initializing tracks"), this);
+	UE_LOG(LogWmfMedia, Verbose, TEXT("Tracks: %p: Initializing (media source %p)"), this, InMediaSource);
 
 	FScopeLock Lock(&CriticalSection);
 
@@ -271,7 +273,7 @@ void FWmfMediaTracks::ReInitialize()
 
 void FWmfMediaTracks::Shutdown()
 {
-	UE_LOG(LogWmfMedia, Verbose, TEXT("Tracks: %p: Shutting down tracks"), this);
+	UE_LOG(LogWmfMedia, Verbose, TEXT("Tracks: %p: Shutting down (media source %p)"), this, MediaSource.Get());
 
 	FScopeLock Lock(&CriticalSection);
 
@@ -293,7 +295,7 @@ void FWmfMediaTracks::Shutdown()
 	if (MediaSource != NULL)
 	{
 		MediaSource->Shutdown();
-		MediaSource = NULL;
+		MediaSource.Reset();
 	}
 
 	PresentationDescriptor.Reset();
@@ -1395,7 +1397,15 @@ bool FWmfMediaTracks::AddStreamToTracks(uint32 StreamIndex, bool IsVideoDevice,
 					else
 					{
 						BufferDim.X = Align(OutputDim.X, 16);
-						BufferDim.Y = Align(OutputDim.Y, 16) * 3 / 2;
+
+						if ((SubType == MFVideoFormat_H264) || (SubType == MFVideoFormat_H264_ES))
+						{
+							BufferDim.Y = Align(OutputDim.Y, 16) * 3 / 2;
+						}
+						else
+						{
+							BufferDim.Y = OutputDim.Y * 3 / 2;
+						}
 					}
 
 					BufferStride = BufferDim.X;
diff --git a/Engine/Plugins/Media/WmfMedia/Source/WmfMedia/Private/Wmf/WmfMediaSession.cpp b/Engine/Plugins/Media/WmfMedia/Source/WmfMedia/Private/Wmf/WmfMediaSession.cpp
--- a/Engine/Plugins/Media/WmfMedia/Source/WmfMedia/Private/Wmf/WmfMediaSession.cpp
+++ b/Engine/Plugins/Media/WmfMedia/Source/WmfMedia/Private/Wmf/WmfMediaSession.cpp
@@ -169,7 +169,7 @@ bool FWmfMediaSession::SetTopology(const TComPtr<IMFTopology>& InTopology, FTime
 		return false;
 	}
 
-	UE_LOG(LogWmfMedia, Verbose, TEXT("Session: %p: Setting new partial topology (duration = %s)"), this, *InDuration.ToString());
+	UE_LOG(LogWmfMedia, Verbose, TEXT("Session: %p: Setting new partial topology %p (duration = %s)"), this, InTopology.Get(), *InDuration.ToString());
 
 	FScopeLock Lock(&CriticalSection);
 
diff --git a/Engine/Plugins/Media/WmfMedia/Source/WmfMedia/Private/Wmf/WmfMediaUtils.cpp b/Engine/Plugins/Media/WmfMedia/Source/WmfMedia/Private/Wmf/WmfMediaUtils.cpp
--- a/Engine/Plugins/Media/WmfMedia/Source/WmfMedia/Private/Wmf/WmfMediaUtils.cpp
+++ b/Engine/Plugins/Media/WmfMedia/Source/WmfMedia/Private/Wmf/WmfMediaUtils.cpp
@@ -493,7 +493,7 @@ namespace WmfMedia
 				}
 			}
 
-			if (SubType == MFVideoFormat_HEVC)
+			if ((SubType == MFVideoFormat_HEVC) || (SubType == MFVideoFormat_HEVC_ES))
 			{
 				if (!FWindowsPlatformMisc::VerifyWindowsVersion(10, 0) /*Win10*/)
 				{
@@ -518,7 +518,10 @@ namespace WmfMedia
 				return NULL;
 			}
 
-			if ((SubType == MFVideoFormat_HEVC) || (SubType == MFVideoFormat_NV12))
+			if ((SubType == MFVideoFormat_HEVC) ||
+				(SubType == MFVideoFormat_HEVC_ES) ||
+				(SubType == MFVideoFormat_NV12) ||
+				(SubType == MFVideoFormat_IYUV))
 			{
 				Result = OutputType->SetGUID(MF_MT_SUBTYPE, MFVideoFormat_NV12);
 			}
@@ -541,11 +544,18 @@ namespace WmfMedia
 			}
 
 			// copy media type attributes
-			if (FAILED(CopyAttribute(&InputType, OutputType, MF_MT_FRAME_RATE)) ||
-				FAILED(CopyAttribute(&InputType, OutputType, MF_MT_FRAME_SIZE)))
+			if (IsVideoDevice)
 			{
-				UE_LOG(LogWmfMedia, Warning, TEXT("Failed to copy video output type attributes"));
-				return NULL;
+				// the following attributes seem to help with web cam issues on Windows 7,
+				// but we generally don't want to copy these for any other media sources
+				// and let the WMF topology resolver pick optimal defaults instead.
+
+				if (FAILED(CopyAttribute(&InputType, OutputType, MF_MT_FRAME_RATE)) ||
+					FAILED(CopyAttribute(&InputType, OutputType, MF_MT_FRAME_SIZE)))
+				{
+					UE_LOG(LogWmfMedia, Warning, TEXT("Failed to copy video output type attributes"));
+					return NULL;
+				}
 			}
 		}
 		else
@@ -754,7 +764,7 @@ namespace WmfMedia
 			Fourcc >>= 8;
 		}
 
-		return Result;
+		return Result.Reverse();
 	}
 
 
diff --git a/Engine/Plugins/Messaging/UdpMessaging/Source/UdpMessaging/Private/Tests/UdpSerializeMessageTaskTest.cpp b/Engine/Plugins/Messaging/UdpMessaging/Source/UdpMessaging/Private/Tests/UdpSerializeMessageTaskTest.cpp
--- a/Engine/Plugins/Messaging/UdpMessaging/Source/UdpMessaging/Private/Tests/UdpSerializeMessageTaskTest.cpp
+++ b/Engine/Plugins/Messaging/UdpMessaging/Source/UdpMessaging/Private/Tests/UdpSerializeMessageTaskTest.cpp
@@ -1,7 +1,7 @@
 // Copyright 1998-2017 Epic Games, Inc. All Rights Reserved.
 
 #include "CoreMinimal.h"
-#include "HAL/PlatformProcess.h"
+#include "HAL/Event.h"
 #include "Misc/AutomationTest.h"
 #include "Async/TaskGraphInterfaces.h"
 #include "Transport/UdpSerializedMessage.h"
@@ -16,38 +16,7 @@ IMPLEMENT_SIMPLE_AUTOMATION_TEST(FUdpSerializeMessageTaskTest, "System.Core.Mess
 
 namespace UdpSerializeMessageTaskTest
 {
-	const int32 NumMessages = 100000;
 	const FTimespan MaxWaitTime(0, 0, 10);
-
-	int32 CompletedMessages = 0;
-	int32 FailedMessages = 0;
-	TSharedPtr<FUdpSerializedMessage, ESPMode::ThreadSafe> ReferenceMessage;
-}
-
-
-void HandleSerializedMessageStateChanged(TSharedRef<FUdpSerializedMessage, ESPMode::ThreadSafe> SerializedMessage)
-{
-	using namespace UdpSerializeMessageTaskTest;
-
-	FPlatformAtomics::InterlockedIncrement(&CompletedMessages);
-
-	const TArray<uint8>& ReferenceDataArray = ReferenceMessage->GetDataArray();
-	const TArray<uint8>& SerializedDataArray = SerializedMessage->GetDataArray();
-
-	if (ReferenceDataArray.Num() != SerializedDataArray.Num())
-	{
-		FPlatformAtomics::InterlockedIncrement(&FailedMessages);
-	}
-	else
-	{
-		const void* ReferenceData = ReferenceDataArray.GetData();
-		const void* SerializedData = SerializedDataArray.GetData();
-
-		if (FMemory::Memcmp(ReferenceData, SerializedData, ReferenceDataArray.Num()) != 0)
-		{
-			FPlatformAtomics::InterlockedIncrement(&FailedMessages);
-		}
-	}
 }
 
 
@@ -55,40 +24,41 @@ bool FUdpSerializeMessageTaskTest::RunTest(const FString& Parameters)
 {
 	using namespace UdpSerializeMessageTaskTest;
 
-	CompletedMessages = 0;
-	FailedMessages = 0;
-	ReferenceMessage = MakeShareable(new FUdpSerializedMessage);
-
-	TSharedRef<IMessageContext, ESPMode::ThreadSafe> Context = MakeShareable(new FUdpMockMessageContext(new FUdpMockMessage));
+	const auto Context = MakeShared<FUdpMockMessageContext, ESPMode::ThreadSafe>(new FUdpMockMessage);
 
 	// synchronous reference serialization
-	FUdpSerializeMessageTask ReferenceTask(Context, ReferenceMessage.ToSharedRef());
+	const auto Message1 = MakeShared<FUdpSerializedMessage, ESPMode::ThreadSafe>();
+
+	FUdpSerializeMessageTask Task1(Context, Message1, nullptr);
 	{
-		ReferenceTask.DoTask(FTaskGraphInterface::Get().GetCurrentThreadIfKnown(), FGraphEventRef());
+		Task1.DoTask(FTaskGraphInterface::Get().GetCurrentThreadIfKnown(), FGraphEventRef());
 	}
 
-	// stress test
-	for (int32 TestIndex = 0; TestIndex < NumMessages; ++TestIndex)
+	// asynchronous serialization
+	TSharedRef<FEvent, ESPMode::ThreadSafe> CompletionEvent = MakeShareable(FPlatformProcess::GetSynchEventFromPool(), [](FEvent* EventToDelete)
 	{
-		TSharedRef<FUdpSerializedMessage, ESPMode::ThreadSafe> SerializedMessage = MakeShareable(new FUdpSerializedMessage);
-		{
-			SerializedMessage->OnStateChanged().BindStatic(&HandleSerializedMessageStateChanged, SerializedMessage);
-		}
+		FPlatformProcess::ReturnSynchEventToPool(EventToDelete);
+	});
 
-		TGraphTask<FUdpSerializeMessageTask>::CreateTask().ConstructAndDispatchWhenReady(Context, SerializedMessage);
-	}
+	TSharedRef<FUdpSerializedMessage, ESPMode::ThreadSafe> Message2 = MakeShareable(new FUdpSerializedMessage);
+	TGraphTask<FUdpSerializeMessageTask>::CreateTask().ConstructAndDispatchWhenReady(Context, Message2, CompletionEvent);
 
-	FDateTime StartTime = FDateTime::UtcNow();
+	const bool Completed = CompletionEvent->Wait(MaxWaitTime);
 
-	while ((CompletedMessages < NumMessages) && ((FDateTime::UtcNow() - StartTime) < MaxWaitTime))
-	{
-		FPlatformProcess::Sleep(0.0f);
-	}
+	const TArray<uint8>& DataArray1 = Message1->GetDataArray();
+	const TArray<uint8>& DataArray2 = Message2->GetDataArray();
+
+	const void* Data1 = DataArray1.GetData();
+	const void* Data2 = DataArray2.GetData();
+
+	const bool Equal = (FMemory::Memcmp(Data1, Data2, DataArray1.Num()) == 0);
 
-	TestEqual(TEXT("The number of completed messages must equal the total number of messages"), CompletedMessages, NumMessages);
-	TestEqual(TEXT("There must be no failed messages"), FailedMessages, (int32)0);
+	TestEqual(TEXT("Synchronous message serialization must succeed"), Message1->GetState(), EUdpSerializedMessageState::Complete);
+	TestTrue(TEXT("Asynchronous message serialization must complete"), Completed);
+	TestEqual(TEXT("Asynchronous message serialization must succeed"), Message2->GetState(), EUdpSerializedMessageState::Complete);
+	TestTrue(TEXT("Synchronous and asynchronous message serialization must yield same results"), Equal);
 
-	return ((CompletedMessages == NumMessages) && (FailedMessages == 0));
+	return (Completed && Equal);
 }
 
 void EmptyLinkFunctionForStaticInitializationUdpSerializeMessageTaskTest()
