MfMediaPlayer: adding flush prior to shutdown during close call to avoid rare hang condition inside the player

diff --git a/Engine/Plugins/Media/MfMedia/Source/MfMedia/Private/Player/MfMediaPlayer.cpp b/Engine/Plugins/Media/MfMedia/Source/MfMedia/Private/Player/MfMediaPlayer.cpp
--- a/Engine/Plugins/Media/MfMedia/Source/MfMedia/Private/Player/MfMediaPlayer.cpp
+++ b/Engine/Plugins/Media/MfMedia/Source/MfMedia/Private/Player/MfMediaPlayer.cpp
@@ -71,7 +71,12 @@ void FMfMediaPlayer::Close()
 	PresentationDescriptor.Reset();
 	RateControl.Reset();
 	RateSupport.Reset();
-	SourceReader.Reset();
+
+	if (SourceReader)
+	{
+		SourceReader->Flush(MF_SOURCE_READER_ALL_STREAMS);
+		SourceReader.Reset();
+	}
 	SourceReaderCallback.Reset();
 
 	Characteristics = 0;
