Adding optional platform-specific MfMedia initialization, allowing per-platform video decoding optimizations.

diff --git a/Engine/Plugins/Media/MfMedia/Source/MfMedia/Private/Player/MfMediaTracks.cpp b/Engine/Plugins/Media/MfMedia/Source/MfMedia/Private/Player/MfMediaTracks.cpp
--- a/Engine/Plugins/Media/MfMedia/Source/MfMedia/Private/Player/MfMediaTracks.cpp
+++ b/Engine/Plugins/Media/MfMedia/Source/MfMedia/Private/Player/MfMediaTracks.cpp
@@ -238,6 +238,12 @@ void FMfMediaTracks::Initialize(IMFMediaSource* InMediaSource, IMFSourceReaderCa
 
 	UE_LOG(LogMfMedia, Verbose, TEXT("Tracks %p: Found %i streams"), this, StreamIndex);
 
+	// Do platform-specific initialization, if required
+#if MFMEDIATRACKS_HAS_PLATFORM_INITIALIZE
+	PlatformInitialize();
+#endif
+
+
 	Algo::Reverse(AudioTracks);
 	Algo::Reverse(CaptionTracks);
 	Algo::Reverse(VideoTracks);
diff --git a/Engine/Plugins/Media/MfMedia/Source/MfMedia/Private/Player/MfMediaTracks.h b/Engine/Plugins/Media/MfMedia/Source/MfMedia/Private/Player/MfMediaTracks.h
--- a/Engine/Plugins/Media/MfMedia/Source/MfMedia/Private/Player/MfMediaTracks.h
+++ b/Engine/Plugins/Media/MfMedia/Source/MfMedia/Private/Player/MfMediaTracks.h
@@ -23,6 +23,10 @@
 	#include "XboxCommonAllowPlatformTypes.h"
 #endif
 
+#if !defined(MFMEDIATRACKS_HAS_PLATFORM_INITIALIZE)
+	#define MFMEDIATRACKS_HAS_PLATFORM_INITIALIZE 0
+#endif
+
 class FMediaSamples;
 class FMfMediaAudioSamplePool;
 class FMfMediaTextureSamplePool; 
@@ -287,6 +291,13 @@ protected:
 	 */
 	void UpdateVideo();
 
+	/** 
+	 * Platform-specific initialization, if requested
+	 */
+#if MFMEDIATRACKS_HAS_PLATFORM_INITIALIZE
+	void PlatformInitialize();
+#endif
+	
 private:
 
 	/** Whether the audio track has reached the end. */
