Remove HoloLens references from MfMedia - moved to platform extension.

#rnx

diff --git a/Engine/Platforms/Hololens/Plugins/Media/MfMedia/MfMedia_HoloLens.uplugin b/Engine/Platforms/Hololens/Plugins/Media/MfMedia/MfMedia_HoloLens.uplugin
new file mode 100644
--- /dev/null
+++ b/Engine/Platforms/Hololens/Plugins/Media/MfMedia/MfMedia_HoloLens.uplugin
@@ -0,0 +1,20 @@
+{
+	"FileVersion": 3,
+	"bIsPluginExtension": true,
+	"Modules": [
+		{
+			"Name": "MfMedia",
+			"Type": "RuntimeNoCommandlet",
+			"PlatformAllowList": [
+				"HoloLens"
+			]
+		},
+		{
+			"Name": "MfMediaFactory",
+			"Type": "RuntimeNoCommandlet",
+			"PlatformAllowList": [
+				"HoloLens"
+			]
+		}
+	]
+}
\ No newline at end of file
diff --git a/Engine/Platforms/Hololens/Plugins/Media/MfMedia/Source/MfMedia/MfMedia_HoloLens.Build.cs b/Engine/Platforms/Hololens/Plugins/Media/MfMedia/Source/MfMedia/MfMedia_HoloLens.Build.cs
new file mode 100644
--- /dev/null
+++ b/Engine/Platforms/Hololens/Plugins/Media/MfMedia/Source/MfMedia/MfMedia_HoloLens.Build.cs
@@ -0,0 +1,13 @@
+// Copyright Epic Games, Inc. All Rights Reserved.
+
+namespace UnrealBuildTool.Rules
+{
+	public class MfMedia_HoloLens : MfMedia
+	{
+		public MfMedia_HoloLens(ReadOnlyTargetRules Target) : base(Target)
+		{
+			PrivateDefinitions.Add("MFMEDIA_SUPPORTED_PLATFORM=1");
+			PrivateDefinitions.Add("MFMEDIA_NEED_PLATFORM_PRIVATE=0");
+		}
+	}
+}
diff --git a/Engine/Plugins/Media/MfMedia/MfMedia.uplugin b/Engine/Plugins/Media/MfMedia/MfMedia.uplugin
--- a/Engine/Plugins/Media/MfMedia/MfMedia.uplugin
+++ b/Engine/Plugins/Media/MfMedia/MfMedia.uplugin
@@ -20,7 +20,8 @@
 			"Name" : "MfMedia",
 			"Type" : "RuntimeNoCommandlet",
 			"LoadingPhase" : "PreLoadingScreen",
-			"PlatformAllowList" : [ "HoloLens" ]
+			"bHasExplicitPlatforms" : true,
+			"PlatformAllowList" : []
 		},
 		{
 			"Name" : "MfMediaEditor",
@@ -36,7 +37,7 @@
 			"Name" : "MfMediaFactory",
 			"Type" : "RuntimeNoCommandlet",
 			"LoadingPhase" : "PostEngineInit",
-			"PlatformAllowList" : [ "HoloLens", "Win64" ]
+			"PlatformAllowList" : [ "Win64" ]
 		}
 	]
 }
diff --git a/Engine/Plugins/Media/MfMedia/Source/MfMedia/Private/Mf/IMfMediaSourceReaderSink.h b/Engine/Plugins/Media/MfMedia/Source/MfMedia/Private/Mf/IMfMediaSourceReaderSink.h
--- a/Engine/Plugins/Media/MfMedia/Source/MfMedia/Private/Mf/IMfMediaSourceReaderSink.h
+++ b/Engine/Plugins/Media/MfMedia/Source/MfMedia/Private/Mf/IMfMediaSourceReaderSink.h
@@ -8,11 +8,9 @@
 
 #include "Misc/Timespan.h"
 
-#if PLATFORM_WINDOWS || PLATFORM_HOLOLENS
-	#include "Windows/WindowsHWrapper.h"
-	#include "Windows/AllowWindowsPlatformTypes.h"
-#else
-	#include "XboxCommonAllowPlatformTypes.h"
+#if PLATFORM_MICROSOFT
+	#include "Microsoft/WindowsHWrapper.h"
+	#include "Microsoft/AllowMicrosoftPlatformTypes.h"
 #endif
 
 
@@ -57,10 +55,8 @@ public:
 };
 
 
-#if PLATFORM_WINDOWS || PLATFORM_HOLOLENS
-	#include "Windows/HideWindowsPlatformTypes.h"
-#else
-	#include "XboxCommonHidePlatformTypes.h"
+#if PLATFORM_MICROSOFT
+	#include "Microsoft/HideMicrosoftPlatformTypes.h"
 #endif
 
 #endif //MFMEDIA_SUPPORTED_PLATFORM
diff --git a/Engine/Plugins/Media/MfMedia/Source/MfMedia/Private/Mf/MfMediaByteStream.cpp b/Engine/Plugins/Media/MfMedia/Source/MfMedia/Private/Mf/MfMediaByteStream.cpp
--- a/Engine/Plugins/Media/MfMedia/Source/MfMedia/Private/Mf/MfMediaByteStream.cpp
+++ b/Engine/Plugins/Media/MfMedia/Source/MfMedia/Private/Mf/MfMediaByteStream.cpp
@@ -11,10 +11,8 @@
 #include "Serialization/Archive.h"
 #include "Microsoft/COMPointer.h"
 
-#if PLATFORM_WINDOWS || PLATFORM_HOLOLENS
-	#include "Windows/AllowWindowsPlatformTypes.h"
-#else
-	#include "XboxCommonAllowPlatformTypes.h"
+#if PLATFORM_MICROSOFT
+	#include "Microsoft/AllowMicrosoftPlatformTypes.h"
 #endif
 
 
@@ -351,10 +349,8 @@ STDMETHODIMP FMfMediaByteStream::Write(const BYTE* pb, ULONG cb, ULONG* pcbWritt
 }
 
 
-#if PLATFORM_WINDOWS || PLATFORM_HOLOLENS
-	#include "Windows/HideWindowsPlatformTypes.h"
-#else
-	#include "XboxCommonHidePlatformTypes.h"
+#if PLATFORM_MICROSOFT
+	#include "Microsoft/HideMicrosoftPlatformTypes.h"
 #endif
 
 #endif //MFMEDIA_SUPPORTED_PLATFORM
diff --git a/Engine/Plugins/Media/MfMedia/Source/MfMedia/Private/Mf/MfMediaByteStream.h b/Engine/Plugins/Media/MfMedia/Source/MfMedia/Private/Mf/MfMediaByteStream.h
--- a/Engine/Plugins/Media/MfMedia/Source/MfMedia/Private/Mf/MfMediaByteStream.h
+++ b/Engine/Plugins/Media/MfMedia/Source/MfMedia/Private/Mf/MfMediaByteStream.h
@@ -9,11 +9,9 @@
 #include "CoreTypes.h"
 #include "Templates/SharedPointer.h"
 
-#if PLATFORM_WINDOWS || PLATFORM_HOLOLENS
-	#include "Windows/WindowsHWrapper.h"
-	#include "Windows/AllowWindowsPlatformTypes.h"
-#else
-	#include "XboxCommonAllowPlatformTypes.h"
+#if PLATFORM_MICROSOFT
+	#include "Microsoft/WindowsHWrapper.h"
+	#include "Microsoft/AllowMicrosoftPlatformTypes.h"
 #endif
 
 class FArchive;
@@ -107,10 +105,8 @@ private:
 };
 
 
-#if PLATFORM_WINDOWS || PLATFORM_HOLOLENS
-#include "Windows/HideWindowsPlatformTypes.h"
-#else
-	#include "XboxCommonHidePlatformTypes.h"
+#if PLATFORM_MICROSOFT
+#include "Microsoft/HideMicrosoftPlatformTypes.h"
 #endif
 
 #endif //MFMEDIA_SUPPORTED_PLATFORM
diff --git a/Engine/Plugins/Media/MfMedia/Source/MfMedia/Private/Mf/MfMediaReadState.h b/Engine/Plugins/Media/MfMedia/Source/MfMedia/Private/Mf/MfMediaReadState.h
--- a/Engine/Plugins/Media/MfMedia/Source/MfMedia/Private/Mf/MfMediaReadState.h
+++ b/Engine/Plugins/Media/MfMedia/Source/MfMedia/Private/Mf/MfMediaReadState.h
@@ -6,11 +6,9 @@
 
 #if MFMEDIA_SUPPORTED_PLATFORM
 
-#if PLATFORM_WINDOWS || PLATFORM_HOLOLENS
-	#include "Windows/WindowsHWrapper.h"
-	#include "Windows/AllowWindowsPlatformTypes.h"
-#else
-	#include "XboxCommonAllowPlatformTypes.h"
+#if PLATFORM_MICROSOFT
+	#include "Microsoft/WindowsHWrapper.h"
+	#include "Microsoft/AllowMicrosoftPlatformTypes.h"
 #endif
 
 
@@ -143,10 +141,8 @@ private:
 };
 
 
-#if PLATFORM_WINDOWS || PLATFORM_HOLOLENS
-	#include "Windows/HideWindowsPlatformTypes.h"
-#else
-	#include "XboxCommonHidePlatformTypes.h"
+#if PLATFORM_MICROSOFT
+	#include "Microsoft/HideMicrosoftPlatformTypes.h"
 #endif
 
 #endif //MFMEDIA_SUPPORTED_PLATFORM
diff --git a/Engine/Plugins/Media/MfMedia/Source/MfMedia/Private/Mf/MfMediaSourceReaderCallback.cpp b/Engine/Plugins/Media/MfMedia/Source/MfMedia/Private/Mf/MfMediaSourceReaderCallback.cpp
--- a/Engine/Plugins/Media/MfMedia/Source/MfMedia/Private/Mf/MfMediaSourceReaderCallback.cpp
+++ b/Engine/Plugins/Media/MfMedia/Source/MfMedia/Private/Mf/MfMediaSourceReaderCallback.cpp
@@ -10,11 +10,9 @@
 #include "IMfMediaSourceReaderSink.h"
 #include "MfMediaUtils.h"
 
-#if PLATFORM_WINDOWS || PLATFORM_HOLOLENS
-	#include "Windows/WindowsHWrapper.h"
-	#include "Windows/AllowWindowsPlatformTypes.h"
-#else
-	#include "XboxCommonAllowPlatformTypes.h"
+#if PLATFORM_MICROSOFT
+	#include "Microsoft/WindowsHWrapper.h"
+	#include "Microsoft/AllowMicrosoftPlatformTypes.h"
 #endif
 
 #define MFMEDIASOURCEREADERCALLBACK_TRACE_SAMPLES 0
@@ -144,10 +142,8 @@ STDMETHODIMP_(ULONG) FMfMediaSourceReaderCallback::Release()
 }
 
 
-#if PLATFORM_WINDOWS || PLATFORM_HOLOLENS
-	#include "Windows/HideWindowsPlatformTypes.h"
-#else
-	#include "XboxCommonHidePlatformTypes.h"
+#if PLATFORM_MICROSOFT
+	#include "Microsoft/HideMicrosoftPlatformTypes.h"
 #endif
 
 #endif //MFMEDIA_SUPPORTED_PLATFORM
diff --git a/Engine/Plugins/Media/MfMedia/Source/MfMedia/Private/Mf/MfMediaSourceReaderCallback.h b/Engine/Plugins/Media/MfMedia/Source/MfMedia/Private/Mf/MfMediaSourceReaderCallback.h
--- a/Engine/Plugins/Media/MfMedia/Source/MfMedia/Private/Mf/MfMediaSourceReaderCallback.h
+++ b/Engine/Plugins/Media/MfMedia/Source/MfMedia/Private/Mf/MfMediaSourceReaderCallback.h
@@ -9,11 +9,9 @@
 #include "CoreTypes.h"
 #include "Templates/SharedPointer.h"
 
-#if PLATFORM_WINDOWS || PLATFORM_HOLOLENS
-	#include "Windows/WindowsHWrapper.h"
-	#include "Windows/AllowWindowsPlatformTypes.h"
-#else
-	#include "XboxCommonAllowPlatformTypes.h"
+#if PLATFORM_MICROSOFT
+	#include "Microsoft/WindowsHWrapper.h"
+	#include "Microsoft/AllowMicrosoftPlatformTypes.h"
 #endif
 
 class IMfMediaSourceReaderSink;
@@ -65,10 +63,8 @@ private:
 };
 
 
-#if PLATFORM_WINDOWS || PLATFORM_HOLOLENS
-	#include "Windows/HideWindowsPlatformTypes.h"
-#else
-	#include "XboxCommonHidePlatformTypes.h"
+#if PLATFORM_MICROSOFT
+	#include "Microsoft/HideMicrosoftPlatformTypes.h"
 #endif
 
 #endif //MFMEDIA_SUPPORTED_PLATFORM
diff --git a/Engine/Plugins/Media/MfMedia/Source/MfMedia/Private/Mf/MfMediaUtils.cpp b/Engine/Plugins/Media/MfMedia/Source/MfMedia/Private/Mf/MfMediaUtils.cpp
--- a/Engine/Plugins/Media/MfMedia/Source/MfMedia/Private/Mf/MfMediaUtils.cpp
+++ b/Engine/Plugins/Media/MfMedia/Source/MfMedia/Private/Mf/MfMediaUtils.cpp
@@ -13,10 +13,8 @@
 #include "MfMediaByteStream.h"
 #include "MfMediaPrivate.h"
 
-#if PLATFORM_WINDOWS || PLATFORM_HOLOLENS
-	#include "Windows/AllowWindowsPlatformTypes.h"
-#else
-	#include "XboxCommonAllowPlatformTypes.h"
+#if PLATFORM_MICROSOFT
+	#include "Microsoft/AllowMicrosoftPlatformTypes.h"
 #endif
 
 #ifndef MFMEDIA_PLATFORM_SUPPORTS_ONLY_AAC
@@ -152,7 +150,7 @@ namespace MfMedia
 					return NULL;
 				}
 			}
-#if !PLATFORM_HOLOLENS
+#if PLATFORM_WINDOWS
 			if ((SubType == MFVideoFormat_H264) || (SubType == MFVideoFormat_H264_ES))
 			{
 				if (!FPlatformMisc::VerifyWindowsVersion(6, 1) /*Win7*/)
@@ -175,7 +173,7 @@ namespace MfMedia
 					UE_LOG(LogMfMedia, Warning, TEXT("HEVC video type requires Windows 10 or newer (your version is %s), and game must be manifested for Windows 10"), *FPlatformMisc::GetOSVersion());
 				}
 			}
-#endif // PLATFORM_HOLOLENS
+#endif // PLATFORM_WINDOWS
 
 #endif //MFMEDIA_PLATFORM_SUPPORTS_ONLY_H264
 
@@ -900,10 +898,8 @@ namespace MfMedia
 }
 
 
-#if PLATFORM_WINDOWS || PLATFORM_HOLOLENS
-	#include "Windows/HideWindowsPlatformTypes.h"
-#else
-	#include "XboxCommonHidePlatformTypes.h"
+#if PLATFORM_MICROSOFT
+	#include "Microsoft/HideMicrosoftPlatformTypes.h"
 #endif
 
 #endif //MFMEDIA_SUPPORTED_PLATFORM
diff --git a/Engine/Plugins/Media/MfMedia/Source/MfMedia/Private/Mf/MfMediaUtils.h b/Engine/Plugins/Media/MfMedia/Source/MfMedia/Private/Mf/MfMediaUtils.h
--- a/Engine/Plugins/Media/MfMedia/Source/MfMedia/Private/Mf/MfMediaUtils.h
+++ b/Engine/Plugins/Media/MfMedia/Source/MfMedia/Private/Mf/MfMediaUtils.h
@@ -9,11 +9,9 @@
 #include "Containers/UnrealString.h"
 #include "Templates/SharedPointer.h"
 
-#if PLATFORM_WINDOWS || PLATFORM_HOLOLENS
-	#include "Windows/WindowsHWrapper.h"
-	#include "Windows/AllowWindowsPlatformTypes.h"
-#else
-	#include "XboxCommonAllowPlatformTypes.h"
+#if PLATFORM_MICROSOFT
+	#include "Microsoft/WindowsHWrapper.h"
+	#include "Microsoft/AllowMicrosoftPlatformTypes.h"
 #endif
 
 class FArchive;
@@ -91,10 +89,8 @@ namespace MfMedia
 }
 
 
-#if PLATFORM_WINDOWS || PLATFORM_HOLOLENS
-	#include "Windows/HideWindowsPlatformTypes.h"
-#else
-	#include "XboxCommonHidePlatformTypes.h"
+#if PLATFORM_MICROSOFT
+	#include "Microsoft/HideMicrosoftPlatformTypes.h"
 #endif
 
 #endif //MFMEDIA_SUPPORTED_PLATFORM
diff --git a/Engine/Plugins/Media/MfMedia/Source/MfMedia/Private/MfMediaPrivate.h b/Engine/Plugins/Media/MfMedia/Source/MfMedia/Private/MfMediaPrivate.h
--- a/Engine/Plugins/Media/MfMedia/Source/MfMedia/Private/MfMediaPrivate.h
+++ b/Engine/Plugins/Media/MfMedia/Source/MfMedia/Private/MfMediaPrivate.h
@@ -5,10 +5,12 @@
 #include "Logging/LogMacros.h"
 
 #if !defined(MFMEDIA_SUPPORTED_PLATFORM)
-	#define MFMEDIA_SUPPORTED_PLATFORM (PLATFORM_HOLOLENS || (PLATFORM_WINDOWS && !UE_SERVER))
+	#define MFMEDIA_SUPPORTED_PLATFORM (PLATFORM_WINDOWS && !UE_SERVER)
 	#define MFMEDIA_NEED_PLATFORM_PRIVATE 0
-#else
+
+#elif !defined(MFMEDIA_NEED_PLATFORM_PRIVATE)
 	#define MFMEDIA_NEED_PLATFORM_PRIVATE 1
+
 #endif
 
 
@@ -18,10 +20,10 @@
 
 #elif MFMEDIA_SUPPORTED_PLATFORM
 	#if PLATFORM_WINDOWS
-		#include "Windows/WindowsHWrapper.h"
+		#include "Microsoft/WindowsHWrapper.h"
 	#endif
 
-	#include "Windows/AllowWindowsPlatformTypes.h"
+	#include "Microsoft/AllowMicrosoftPlatformTypes.h"
 
 	#if PLATFORM_WINDOWS
 		#include <windows.h>
@@ -42,7 +44,7 @@
 	#include <Mfreadwrite.h>
 
 	#include "Microsoft/COMPointer.h"
-	#include "Windows/HideWindowsPlatformTypes.h"
+	#include "Microsoft/HideMicrosoftPlatformTypes.h"
 
 #elif PLATFORM_WINDOWS && !UE_SERVER
 	#pragma message("Skipping MfMedia (requires WINVER >= 0x0601, but WINVER is " PREPROCESSOR_TO_STRING(WINVER) ")")
diff --git a/Engine/Plugins/Media/MfMedia/Source/MfMedia/Private/Player/MfMediaPlayer.cpp b/Engine/Plugins/Media/MfMedia/Source/MfMedia/Private/Player/MfMediaPlayer.cpp
--- a/Engine/Plugins/Media/MfMedia/Source/MfMedia/Private/Player/MfMediaPlayer.cpp
+++ b/Engine/Plugins/Media/MfMedia/Source/MfMedia/Private/Player/MfMediaPlayer.cpp
@@ -14,11 +14,9 @@
 #include "MfMediaTracks.h"
 #include "MfMediaUtils.h"
 
-#if PLATFORM_WINDOWS || PLATFORM_HOLOLENS
-	#include "Windows/WindowsHWrapper.h"
-	#include "Windows/AllowWindowsPlatformTypes.h"
-#else
-	#include "XboxCommonAllowPlatformTypes.h"
+#if PLATFORM_MICROSOFT
+	#include "Microsoft/WindowsHWrapper.h"
+	#include "Microsoft/AllowMicrosoftPlatformTypes.h"
 #endif
 
 #define LOCTEXT_NAMESPACE "FMfMediaPlayer"
@@ -768,10 +766,8 @@ void FMfMediaPlayer::ReceiveSourceReaderSample(IMFSample* Sample, HRESULT Status
 
 #undef LOCTEXT_NAMESPACE
 
-#if PLATFORM_WINDOWS || PLATFORM_HOLOLENS
-	#include "Windows/HideWindowsPlatformTypes.h"
-#else
-	#include "XboxCommonHidePlatformTypes.h"
+#if PLATFORM_MICROSOFT
+	#include "Microsoft/HideMicrosoftPlatformTypes.h"
 #endif
 
 #endif //MFMEDIA_SUPPORTED_PLATFORM
diff --git a/Engine/Plugins/Media/MfMedia/Source/MfMedia/Private/Player/MfMediaPlayer.h b/Engine/Plugins/Media/MfMedia/Source/MfMedia/Private/Player/MfMediaPlayer.h
--- a/Engine/Plugins/Media/MfMedia/Source/MfMedia/Private/Player/MfMediaPlayer.h
+++ b/Engine/Plugins/Media/MfMedia/Source/MfMedia/Private/Player/MfMediaPlayer.h
@@ -17,11 +17,9 @@
 
 #include "IMfMediaSourceReaderSink.h"
 
-#if PLATFORM_WINDOWS || PLATFORM_HOLOLENS
-	#include "Windows/WindowsHWrapper.h"
-	#include "Windows/AllowWindowsPlatformTypes.h"
-#else
-	#include "XboxCommonAllowPlatformTypes.h"
+#if PLATFORM_MICROSOFT
+	#include "Microsoft/WindowsHWrapper.h"
+	#include "Microsoft/AllowMicrosoftPlatformTypes.h"
 #endif
 
 class FMediaSamples;
@@ -186,10 +184,8 @@ private:
 };
 
 
-#if PLATFORM_WINDOWS || PLATFORM_HOLOLENS
-	#include "Windows/HideWindowsPlatformTypes.h"
-#else
-	#include "XboxCommonHidePlatformTypes.h"
+#if PLATFORM_MICROSOFT
+	#include "Microsoft/HideMicrosoftPlatformTypes.h"
 #endif
 
 #endif //MFMEDIA_SUPPORTED_PLATFORM
diff --git a/Engine/Plugins/Media/MfMedia/Source/MfMedia/Private/Player/MfMediaTracks.cpp b/Engine/Plugins/Media/MfMedia/Source/MfMedia/Private/Player/MfMediaTracks.cpp
--- a/Engine/Plugins/Media/MfMedia/Source/MfMedia/Private/Player/MfMediaTracks.cpp
+++ b/Engine/Plugins/Media/MfMedia/Source/MfMedia/Private/Player/MfMediaTracks.cpp
@@ -18,11 +18,9 @@
 #include "MfMediaTextureSample.h"
 #include "MfMediaUtils.h"
 
-#if PLATFORM_WINDOWS || PLATFORM_HOLOLENS
-	#include "Windows/WindowsHWrapper.h"
-	#include "Windows/AllowWindowsPlatformTypes.h"
-#else
-	#include "XboxCommonAllowPlatformTypes.h"
+#if PLATFORM_MICROSOFT
+	#include "Microsoft/WindowsHWrapper.h"
+	#include "Microsoft/AllowMicrosoftPlatformTypes.h"
 #endif
 
 #define MFMEDIATRACKS_TRACE_SAMPLES 0
@@ -1568,10 +1566,8 @@ void FMfMediaTracks::UpdateVideo()
 
 #undef LOCTEXT_NAMESPACE
 
-#if PLATFORM_WINDOWS || PLATFORM_HOLOLENS
-	#include "Windows/HideWindowsPlatformTypes.h"
-#else
-	#include "XboxCommonHidePlatformTypes.h"
+#if PLATFORM_MICROSOFT
+	#include "Microsoft/HideMicrosoftPlatformTypes.h"
 #endif
 
 #endif //MFMEDIA_SUPPORTED_PLATFORM
diff --git a/Engine/Plugins/Media/MfMedia/Source/MfMedia/Private/Player/MfMediaTracks.h b/Engine/Plugins/Media/MfMedia/Source/MfMedia/Private/Player/MfMediaTracks.h
--- a/Engine/Plugins/Media/MfMedia/Source/MfMedia/Private/Player/MfMediaTracks.h
+++ b/Engine/Plugins/Media/MfMedia/Source/MfMedia/Private/Player/MfMediaTracks.h
@@ -16,11 +16,9 @@
 #include "Math/Range.h"
 #include "Microsoft/COMPointer.h"
 
-#if PLATFORM_WINDOWS || PLATFORM_HOLOLENS
-	#include "Windows/WindowsHWrapper.h"
-	#include "Windows/AllowWindowsPlatformTypes.h"
-#else
-	#include "XboxCommonAllowPlatformTypes.h"
+#if PLATFORM_MICROSOFT
+	#include "Microsoft/WindowsHWrapper.h"
+	#include "Microsoft/AllowMicrosoftPlatformTypes.h"
 #endif
 
 #if !defined(MFMEDIATRACKS_HAS_PLATFORM_INITIALIZE)
@@ -383,10 +381,8 @@ private:
 };
 
 
-#if PLATFORM_WINDOWS || PLATFORM_HOLOLENS
-	#include "Windows/HideWindowsPlatformTypes.h"
-#else
-	#include "XboxCommonHidePlatformTypes.h"
+#if PLATFORM_MICROSOFT
+	#include "Microsoft/HideMicrosoftPlatformTypes.h"
 #endif
 
 #endif //MFMEDIA_SUPPORTED_PLATFORM
