Restore HoloLens platform reference for MfMedia to fix the editor. Needs further investigation.

#rnx

diff --git a/Engine/Plugins/Media/MfMedia/MfMedia.uplugin b/Engine/Plugins/Media/MfMedia/MfMedia.uplugin
--- a/Engine/Plugins/Media/MfMedia/MfMedia.uplugin
+++ b/Engine/Plugins/Media/MfMedia/MfMedia.uplugin
@@ -20,8 +20,7 @@
 			"Name" : "MfMedia",
 			"Type" : "RuntimeNoCommandlet",
 			"LoadingPhase" : "PreLoadingScreen",
-			"bHasExplicitPlatforms" : true,
-			"PlatformAllowList" : []
+			"PlatformAllowList" : [ "HoloLens" ]
 		},
 		{
 			"Name" : "MfMediaEditor",
