Remove HoloLens platform reference for MfMedia again & use the correct field name this time.

#rnx

diff --git a/Engine/Plugins/Media/MfMedia/MfMedia.uplugin b/Engine/Plugins/Media/MfMedia/MfMedia.uplugin
--- a/Engine/Plugins/Media/MfMedia/MfMedia.uplugin
+++ b/Engine/Plugins/Media/MfMedia/MfMedia.uplugin
@@ -20,7 +20,8 @@
 			"Name" : "MfMedia",
 			"Type" : "RuntimeNoCommandlet",
 			"LoadingPhase" : "PreLoadingScreen",
-			"PlatformAllowList" : [ "HoloLens" ]
+			"HasExplicitPlatforms" : true,
+			"PlatformAllowList" : []
 		},
 		{
 			"Name" : "MfMediaEditor",
