Moving WMF-triggered flush processing to the main gamethread for MFMediaPlayer to avoid assertion (and race conditions)

diff --git a/Engine/Plugins/Media/MfMedia/Source/MfMedia/Private/Player/MfMediaPlayer.cpp b/Engine/Plugins/Media/MfMedia/Source/MfMedia/Private/Player/MfMediaPlayer.cpp
--- a/Engine/Plugins/Media/MfMedia/Source/MfMedia/Private/Player/MfMediaPlayer.cpp
+++ b/Engine/Plugins/Media/MfMedia/Source/MfMedia/Private/Player/MfMediaPlayer.cpp
@@ -749,7 +749,20 @@ void FMfMediaPlayer::ReceiveSourceReaderEvent(MediaEventType Event)
 
 void FMfMediaPlayer::ReceiveSourceReaderFlush()
 {
-	Samples->FlushSamples();
+	if (IsInGameThread() || IsInSlateThread())
+	{
+		Samples->FlushSamples();
+	}
+	else
+	{
+		TWeakPtr<FMediaSamples, ESPMode::ThreadSafe> WeakSamples(Samples);
+		Async(EAsyncExecution::TaskGraphMainThread, [WeakSamples]() {
+			if (auto PinnedSamples = WeakSamples.Pin())
+			{
+				PinnedSamples->FlushSamples();
+			}
+			});
+	}
 }
 
 
