fix some Clang compile warnings in MfMedia code.

#rnx

diff --git a/Engine/Plugins/Media/MfMedia/Source/MfMedia/Private/Player/MfMediaPlayer.cpp b/Engine/Plugins/Media/MfMedia/Source/MfMedia/Private/Player/MfMediaPlayer.cpp
--- a/Engine/Plugins/Media/MfMedia/Source/MfMedia/Private/Player/MfMediaPlayer.cpp
+++ b/Engine/Plugins/Media/MfMedia/Source/MfMedia/Private/Player/MfMediaPlayer.cpp
@@ -52,7 +52,7 @@ FMfMediaPlayer::~FMfMediaPlayer()
 
 void FMfMediaPlayer::Close()
 {
-	if ((CurrentState == EMediaState::Closed))
+	if (CurrentState == EMediaState::Closed)
 	{
 		return;
 	}
@@ -397,8 +397,8 @@ bool FMfMediaPlayer::InitializePlayer(const TSharedPtr<FArchive, ESPMode::Thread
 
 		if (PinnedSamples.IsValid() && PinnedTracks.IsValid())
 		{
-			TComPtr<IMFMediaSource> MediaSource = MfMedia::ResolveMediaSource(Archive, Url, Precache);
-			PinnedTracks->Initialize(MediaSource, Callback, PinnedSamples.ToSharedRef(), &LocalPlayerOptions);
+			TComPtr<IMFMediaSource> MediaSourcePtr = MfMedia::ResolveMediaSource(Archive, Url, Precache);
+			PinnedTracks->Initialize(MediaSourcePtr, Callback, PinnedSamples.ToSharedRef(), &LocalPlayerOptions);
 		}
 	});
 
